Automation - Improving report generation and adjusting the way we generate the expors for screenshot comparison failure data to also go to the report server, so it's all in one easy to manage place.  Automation report now also lets users know where it can manage the screenshots comparison failures.

diff --git a/Engine/Plugins/Tests/ScreenshotTools/Source/ScreenshotTools/Private/CompareScreenshotsCommandlet.cpp b/Engine/Plugins/Tests/ScreenshotTools/Source/ScreenshotTools/Private/CompareScreenshotsCommandlet.cpp
deleted file mode 100644
diff --git a/Engine/Plugins/Tests/ScreenshotTools/Source/ScreenshotTools/Private/CompareScreenshotsCommandlet.h b/Engine/Plugins/Tests/ScreenshotTools/Source/ScreenshotTools/Private/CompareScreenshotsCommandlet.h
deleted file mode 100644
diff --git a/Engine/Source/Developer/AutomationController/Private/AutomationControllerManger.cpp b/Engine/Source/Developer/AutomationController/Private/AutomationControllerManger.cpp
--- a/Engine/Source/Developer/AutomationController/Private/AutomationControllerManger.cpp
+++ b/Engine/Source/Developer/AutomationController/Private/AutomationControllerManger.cpp
@@ -353,6 +353,7 @@ void FAutomationControllerManager::GenerateHtmlTestPassSummary(FDateTime Timesta
 	}
 
 	FString ReportOutputPath = GetReportPath(Timestamp);
+	FScreenshotExportResults ExportResults = ScreenshotManager->ExportComparisonResultsAsync(ReportOutputPath).Get();
 
 	FAutomatedTestPassResults SerializedPassResults = OurPassResults;
 	SerializedPassResults.TestInformation.StableSort([] (const FAutomatedTestResult& A, const FAutomatedTestResult& B) {
@@ -471,6 +472,7 @@ void FAutomationControllerManager::GenerateHtmlTestPassSummary(FDateTime Timesta
 		TMap<FString, FStringFormatArg> Args;
 		Args.Add(TEXT("Title"), TEXT("Automation Test Results"));
 		Args.Add(TEXT("TitleColor"), TitleColor);
+		Args.Add(TEXT("ComparisonExportDirectory"), ExportResults.ExportPath);
 		Args.Add(TEXT("Results"), HtmlResults);
 
 		FString Html = FString::Format(*MasterTemplate, Args);
@@ -712,11 +714,17 @@ void FAutomationControllerManager::ProcessResults()
 		}
 	}
 
-	// Write our the report of our automation pass to JSON. Then clean our array for the next pass.
-	FDateTime Timestamp = FDateTime::Now();
-	GenerateJsonTestPassSummary(Timestamp);
-	GenerateHtmlTestPassSummary(Timestamp);
+	if ( !ReportOutputPathOverride.IsEmpty() )
+	{
+		FDateTime Timestamp = FDateTime::Now();
+
+		// Generate Html
+		GenerateHtmlTestPassSummary(Timestamp);
+		// Generate Json
+		GenerateJsonTestPassSummary(Timestamp);
+	}
 
+	// Then clean our array for the next pass.
 	OurPassResults.ClearAllEntries();
 
 	SetControllerStatus(EAutomationControllerModuleState::Ready);
diff --git a/Engine/Source/Developer/ScreenShotComparisonTools/Private/ScreenShotManager.cpp b/Engine/Source/Developer/ScreenShotComparisonTools/Private/ScreenShotManager.cpp
--- a/Engine/Source/Developer/ScreenShotComparisonTools/Private/ScreenShotManager.cpp
+++ b/Engine/Source/Developer/ScreenShotComparisonTools/Private/ScreenShotManager.cpp
@@ -217,12 +217,12 @@ FImageComparisonResult FScreenShotManager::CompareScreensot(FString ExistingImag
 	return ComparisonResult;
 }
 
-TFuture<void> FScreenShotManager::ExportScreensotsAsync(FString ExportPath)
+TFuture<FScreenshotExportResults> FScreenShotManager::ExportComparisonResultsAsync(FString ExportPath)
 {
-	return Async<void>(EAsyncExecution::Thread, [&] () { ExportComparisonResults(ExportPath); });
+	return Async<FScreenshotExportResults>(EAsyncExecution::Thread, [&] () { return ExportComparisonResults(ExportPath); });
 }
 
-bool FScreenShotManager::ExportComparisonResults(FString RootExportFolder)
+FScreenshotExportResults FScreenShotManager::ExportComparisonResults(FString RootExportFolder)
 {
 	FPaths::NormalizeDirectoryName(RootExportFolder);
 
@@ -231,19 +231,22 @@ bool FScreenShotManager::ExportComparisonResults(FString RootExportFolder)
 		RootExportFolder = GetDefaultExportDirectory();
 	}
 
-	FString ExportDirectory = RootExportFolder / FString::FromInt(FEngineVersion::Current().GetChangelist());
+	FScreenshotExportResults Results;
+	Results.Success = false;
+	Results.ExportPath = RootExportFolder / FString::FromInt(FEngineVersion::Current().GetChangelist());
 
-	if ( !IFileManager::Get().MakeDirectory(*ExportDirectory, /*Tree =*/true) )
+	if ( !IFileManager::Get().MakeDirectory(*Results.ExportPath, /*Tree =*/true) )
 	{
-		return false;
+		return Results;
 	}
 
 	// Wait for file operations to complete.
 	FPlatformProcess::Sleep(1.0f);
 
-	CopyDirectory(ExportDirectory, ComparisonResultsFolder);
+	CopyDirectory(Results.ExportPath, ComparisonResultsFolder);
 
-	return true;
+	Results.Success = true;
+	return Results;
 }
 
 bool FScreenShotManager::OpenComparisonReports(FString ImportPath, TArray<FComparisonReport>& OutReports)
diff --git a/Engine/Source/Developer/ScreenShotComparisonTools/Private/ScreenShotManager.h b/Engine/Source/Developer/ScreenShotComparisonTools/Private/ScreenShotManager.h
--- a/Engine/Source/Developer/ScreenShotComparisonTools/Private/ScreenShotManager.h
+++ b/Engine/Source/Developer/ScreenShotComparisonTools/Private/ScreenShotManager.h
@@ -35,7 +35,7 @@ public:
 
 	virtual TFuture<FImageComparisonResult> CompareScreensotAsync(FString RelativeImagePath) override;
 
-	virtual TFuture<void> ExportScreensotsAsync(FString ExportPath = TEXT("")) override;
+	virtual TFuture<FScreenshotExportResults> ExportComparisonResultsAsync(FString ExportPath = TEXT("")) override;
 
 	virtual bool OpenComparisonReports(FString ImportPath, TArray<FComparisonReport>& OutReports) override;
 
@@ -51,7 +51,7 @@ private:
 	FString GetDefaultExportDirectory() const;
 
 	FImageComparisonResult CompareScreensot(FString Existing);
-	bool ExportComparisonResults(FString RootExportFolder);
+	FScreenshotExportResults ExportComparisonResults(FString RootExportFolder);
 	void CopyDirectory(const FString& DestDir, const FString& SrcDir);
 
 private:
diff --git a/Engine/Source/Developer/ScreenShotComparisonTools/Public/Interfaces/IScreenShotManager.h b/Engine/Source/Developer/ScreenShotComparisonTools/Public/Interfaces/IScreenShotManager.h
--- a/Engine/Source/Developer/ScreenShotComparisonTools/Public/Interfaces/IScreenShotManager.h
+++ b/Engine/Source/Developer/ScreenShotComparisonTools/Public/Interfaces/IScreenShotManager.h
@@ -22,6 +22,17 @@ typedef TSharedPtr<class IScreenShotManager> IScreenShotManagerPtr;
  */
 typedef TSharedRef<class IScreenShotManager> IScreenShotManagerRef;
 
+struct FScreenshotExportResults
+{
+	bool Success;
+	FString ExportPath;
+
+	FScreenshotExportResults()
+		: Success(false)
+	{
+	}
+};
+
 /**
  * Interface for screen manager module.
  */
@@ -38,7 +49,7 @@ public:
 	/**
 	 * Exports the screenshots to the export location specified
 	 */
-	virtual TFuture<void> ExportScreensotsAsync(FString ExportPath = TEXT("")) = 0;
+	virtual TFuture<FScreenshotExportResults> ExportComparisonResultsAsync(FString ExportPath = TEXT("")) = 0;
 
 	/**
 	 * Imports screenshot comparison data from a given path.
