Addressed image comparison using wrong image with EStructuralSimilarity::Color.

#4806

diff --git a/Engine/Source/Developer/ScreenShotComparisonTools/Private/ImageComparer.cpp b/Engine/Source/Developer/ScreenShotComparisonTools/Private/ImageComparer.cpp
--- a/Engine/Source/Developer/ScreenShotComparisonTools/Private/ImageComparer.cpp
+++ b/Engine/Source/Developer/ScreenShotComparisonTools/Private/ImageComparer.cpp
@@ -507,7 +507,7 @@ double FImageComparer::CompareStructuralSimilarity(const FString& ImagePathA, co
 					else // EStructuralSimilarityComponent::Color
 					{
 						const FColor ColorA = ImageA->GetPixel(WindowX, WindowY);
-						const FColor ColorB = ImageA->GetPixel(WindowX, WindowY);
+						const FColor ColorB = ImageB->GetPixel(WindowX, WindowY);
 						const double ColorLumpA = ( ColorA.R + ColorA.G + ColorA.B ) * ( ColorA.A / 255.0 );
 						const double ColorLumpB = ( ColorB.R + ColorB.G + ColorB.B ) * ( ColorB.A / 255.0 );
 						AverageA += ColorLumpA;
