PR #6037: Fix compiler error of "symbol not found" (Contributed by [REDACTED_PERSON])

diff --git a/Engine/Source/Developer/ScreenShotComparisonTools/Public/ImageComparer.h b/Engine/Source/Developer/ScreenShotComparisonTools/Public/ImageComparer.h
--- a/Engine/Source/Developer/ScreenShotComparisonTools/Public/ImageComparer.h
+++ b/Engine/Source/Developer/ScreenShotComparisonTools/Public/ImageComparer.h
@@ -13,7 +13,7 @@ class FComparableImage;
  * 
  */
 USTRUCT()
-struct FImageTolerance
+struct SCREENSHOTCOMPARISONTOOLS_API FImageTolerance
 {
 	GENERATED_USTRUCT_BODY()
 
