Edigrate cl 16313458, 16342118
Make automated tests report more resilient in case of critical failure or timeout.
* Json report saved after each test execution
* Test artifacts are copied after each test execution
* Update html report to display `not run` and `in process` test state
* Remove undocumented command line option `checkpoint`

#rnx

diff --git a/Engine/Source/Developer/AutomationController/Private/AutomationControllerManager.h b/Engine/Source/Developer/AutomationController/Private/AutomationControllerManager.h
--- a/Engine/Source/Developer/AutomationController/Private/AutomationControllerManager.h
+++ b/Engine/Source/Developer/AutomationController/Private/AutomationControllerManager.h
@@ -608,7 +608,7 @@ private:
 	struct FComparisonEntry
 	{
 		FMessageAddress Sender;
-		FString TestName;
+		FString ScreenshotPath;
 		TFuture<FImageComparisonResult> PendingComparison;
 	};
 
diff --git a/Engine/Source/Developer/AutomationController/Private/AutomationControllerManger.cpp b/Engine/Source/Developer/AutomationController/Private/AutomationControllerManger.cpp
--- a/Engine/Source/Developer/AutomationController/Private/AutomationControllerManger.cpp
+++ b/Engine/Source/Developer/AutomationController/Private/AutomationControllerManger.cpp
@@ -336,7 +336,7 @@ void FAutomationControllerManager::ProcessComparisonQueue()
 					LocalFiles.Add(TEXT("difference"), FPaths::Combine(ProjectDir, Result.ComparisonFilePath));
 				}
 
-				Report->AddArtifact(ClusterIndex, CurrentTestPass, FAutomationArtifact(UniqueId, Entry->TestName, EAutomationArtifactType::Comparison, LocalFiles));
+				Report->AddArtifact(ClusterIndex, CurrentTestPass, FAutomationArtifact(UniqueId, Entry->ScreenshotPath, EAutomationArtifactType::Comparison, LocalFiles));
 			}
 			else
 			{
@@ -1070,7 +1070,7 @@ void FAutomationControllerManager::HandleReceivedScreenShot(const FAutomationWor
 	// compare the incoming image and throw it away afterwards (note - there will be a copy in the report)
 	TSharedRef<FComparisonEntry> Comparison = MakeShareable(new FComparisonEntry());
 	Comparison->Sender = Context->GetSender();
-	Comparison->TestName = Message.Metadata.Context / Message.Metadata.TestName;
+	Comparison->ScreenshotPath = Message.Metadata.Context /  Message.Metadata.ScreenShotName;
 	Comparison->PendingComparison = ScreenshotManager->CompareScreenshotAsync(IncomingFileName, Message.Metadata, EScreenShotCompareOptions::DiscardImage);
 
 	ComparisonQueue.Enqueue(Comparison);
diff --git a/Engine/Source/Developer/ScreenShotComparisonTools/Private/ScreenShotManager.cpp b/Engine/Source/Developer/ScreenShotComparisonTools/Private/ScreenShotManager.cpp
--- a/Engine/Source/Developer/ScreenShotComparisonTools/Private/ScreenShotManager.cpp
+++ b/Engine/Source/Developer/ScreenShotComparisonTools/Private/ScreenShotManager.cpp
@@ -76,16 +76,11 @@ FString FScreenShotManager::GetPathComponentForPlatformAndRHI(const FAutomationS
 }
 
 /**
- * Images are now preferred to be under MapOrContext/TestName/ImageName/Plat/RHI etc
+ * Images are now preferred to be under MapOrContext/ImageName/Plat/RHI etc
  */
 FString FScreenShotManager::GetPathComponentForTestImages(const FAutomationScreenshotMetadata& MetaData) const
 {
-	//if (bUseLegacyFormat)
-	{
-		return FPaths::Combine(MetaData.Context, *MetaData.ScreenShotName);
-	}
-
-	//return FPaths::Combine(MetaData.Context, MetaData.TestName, MetaData.ScreenShotName);
+	return FPaths::Combine(MetaData.Context, *MetaData.ScreenShotName);
 }
 
 FString FScreenShotManager::GetApprovedFolderForImageWithOptions(const FAutomationScreenshotMetadata& MetaData, EApprovedFolderOptions InOptions) const
@@ -96,7 +91,7 @@ FString FScreenShotManager::GetApprovedFolderForImageWithOptions(const FAutomati
 
 	bool bUsePlatformPath = PlatInfo.bIsConfidential && (InOptions & EApprovedFolderOptions::UsePlatformFolders) == 0;
 
-	// Test folder will be either Map/Map_ImageName/ for legacy or Map/Test/ImageName/
+	// Test folder will be MapOrContext/ImageName
 	FString TestFolder = GetPathComponentForTestImages(MetaData);
 
 	FString OutPath = FPaths::ProjectDir();
