Remove FComparableImage::Process() and its outputs (Total/Average values).

The outputs were computed incorrectly (due to data races in ParallelFor body) and inefficiently (looping over pixel columns). Luckily, outputs were not referenced by any subsequent code.

Saves ~0.3 seconds of processing time (~30%, ~1.0s sec -> ~0.7s) per screenshot test on AMD TR 3970X.

diff --git a/Engine/Source/Developer/ScreenShotComparisonTools/Private/ImageComparer.cpp b/Engine/Source/Developer/ScreenShotComparisonTools/Private/ImageComparer.cpp
--- a/Engine/Source/Developer/ScreenShotComparisonTools/Private/ImageComparer.cpp
+++ b/Engine/Source/Developer/ScreenShotComparisonTools/Private/ImageComparer.cpp
@@ -212,33 +212,6 @@ FComparisonReport::FComparisonReport(const FString& InReportRootDirectory, const
 	ReportPath = FPaths::GetPath(InReportFile);
 }
 
-void FComparableImage::Process()
-{
-	ParallelFor(Width,
-		[&] (int32 ColumnIndex)
-	{
-		for ( int Y = 0; Y < Height; Y++ )
-		{
-			FColor Pixel = GetPixel(ColumnIndex, Y);
-			double Luminance = FPixelOperations::GetLuminance(Pixel);
-
-			RedTotal += ( Pixel.R / 255.0 );
-			GreenTotal += ( Pixel.G / 255.0 );
-			BlueTotal += ( Pixel.B / 255.0 );
-			AlphaTotal += ( Pixel.A / 255.0 );
-			LuminanceTotal += ( Luminance / 255.0 );
-		}
-	});
-
-	const double PixelCount = Width * Height;
-
-	RedAverage = RedTotal / PixelCount;
-	GreenAverage = GreenTotal / PixelCount;
-	BlueAverage = BlueTotal / PixelCount;
-	AlphaAverage = AlphaTotal / PixelCount;
-	LuminanceAverage = LuminanceTotal / PixelCount;
-}
-
 TSharedPtr<FComparableImage> FImageComparer::Open(const FString& ImagePath, FText& OutError)
 {
 	const FString ImageExtension = FPaths::GetExtension(ImagePath);
@@ -308,9 +281,6 @@ FImageComparisonResult FImageComparer::Compare(const FString& ImagePathA, const
 		return Results;
 	}
 
-	ImageA->Process();
-	ImageB->Process();
-
 	FImageComparisonResult Results = Compare(ImageA.Get(), ImageB.Get(), Tolerance, OutDeltaPath);
 	Results.ApprovedFilePath = ImagePathA;
 	Results.IncomingFilePath = ImagePathB;
diff --git a/Engine/Source/Developer/ScreenShotComparisonTools/Public/ImageComparer.h b/Engine/Source/Developer/ScreenShotComparisonTools/Public/ImageComparer.h
--- a/Engine/Source/Developer/ScreenShotComparisonTools/Public/ImageComparer.h
+++ b/Engine/Source/Developer/ScreenShotComparisonTools/Public/ImageComparer.h
@@ -143,21 +143,11 @@ public:
 class FComparableImage
 {
 public:
-	int32 Width;
-	int32 Height;
+	int32 Width = 0;
+	int32 Height = 0;
 	TArray64<uint8> Bytes;
 
 	FComparableImage()
-		: RedTotal(0)
-		, GreenTotal(0)
-		, BlueTotal(0)
-		, AlphaTotal(0)
-		, LuminanceTotal(0)
-		, RedAverage(0)
-		, GreenAverage(0)
-		, BlueAverage(0)
-		, AlphaAverage(0)
-		, LuminanceAverage(0)
 	{
 	}
 
@@ -177,22 +167,6 @@ public:
 			Bytes[Offset + 2],
 			Bytes[Offset + 3]);
 	}
-
-	void Process();
-
-public:
-	// Processed Data
-	double RedTotal;
-	double GreenTotal;
-	double BlueTotal;
-	double AlphaTotal;
-	double LuminanceTotal;
-
-	double RedAverage;
-	double GreenAverage;
-	double BlueAverage;
-	double AlphaAverage;
-	double LuminanceAverage;
 };
 
 /**
