Properly initiate `FImageComparisonResult::CreationTime` member
Fix originaly proposed through PR #8632. But the fix is not to exclude the class member to be tested but to properly intiate it.

diff --git a/Engine/Source/Developer/ScreenShotComparisonTools/Private/ScreenShotManager.cpp b/Engine/Source/Developer/ScreenShotComparisonTools/Private/ScreenShotManager.cpp
--- a/Engine/Source/Developer/ScreenShotComparisonTools/Private/ScreenShotManager.cpp
+++ b/Engine/Source/Developer/ScreenShotComparisonTools/Private/ScreenShotManager.cpp
@@ -336,6 +336,7 @@ FImageComparisonResult FScreenShotManager::CompareScreenshot(const FString& InUn
 	{
 		// We can't find a ground truth, so it's a new comparison.
 		ComparisonResult.IncomingFilePath = InUnapprovedIncomingFilePath;
+		ComparisonResult.CreationTime = FDateTime::Now();
 
 		UE_LOG(LogScreenShotManager, Log, TEXT("No ideal-image found. Assuming %s is a new test image"), *InUnapprovedIncomingFilePath);
 	}
@@ -373,7 +374,7 @@ FImageComparisonResult FScreenShotManager::CompareScreenshot(const FString& InUn
 	*/
 	
 	// files we write/copy to for reports
-	ComparisonResult.ReportApprovedFilePath = FPaths::Combine(ReportPathOnDisk, TEXT("Approved.png"));
+	ComparisonResult.ReportApprovedFilePath = FPaths::Combine(ReportPathOnDisk, TEXT("Approved.png")); 
 	ComparisonResult.ReportIncomingFilePath = FPaths::Combine(ReportPathOnDisk, TEXT("Incoming.png"));
 	ComparisonResult.ReportComparisonFilePath = FPaths::Combine(ReportPathOnDisk, TEXT("Delta.png"));
 
diff --git a/Engine/Source/Developer/ScreenShotComparisonTools/Public/ImageComparer.h b/Engine/Source/Developer/ScreenShotComparisonTools/Public/ImageComparer.h
--- a/Engine/Source/Developer/ScreenShotComparisonTools/Public/ImageComparer.h
+++ b/Engine/Source/Developer/ScreenShotComparisonTools/Public/ImageComparer.h
@@ -280,7 +280,7 @@ public:
 	static constexpr int32 CurrentVersion = 2;
 
 	FImageComparisonResult()
-		: CreationTime(FDateTime::Now())
+		: CreationTime(0)
 		, MaxLocalDifference(0.0f)
 		, GlobalDifference(0.0f)
 		, ErrorMessage()
@@ -289,7 +289,7 @@ public:
 	}
 
 	FImageComparisonResult(const FText& Error)
-		: CreationTime(FDateTime::Now())
+		: CreationTime(0)
 		, MaxLocalDifference(0.0f)
 		, GlobalDifference(0.0f)
 		, ErrorMessage(Error)
