Automation: Exposed "GetIdealApprovedFolderForImage" to IScreenShotManager so that external modules can reuse the same code to correctly calculate where to store a ground truth image.

diff --git a/Engine/Source/Developer/ScreenShotComparisonTools/Private/ScreenShotManager.h b/Engine/Source/Developer/ScreenShotComparisonTools/Private/ScreenShotManager.h
--- a/Engine/Source/Developer/ScreenShotComparisonTools/Private/ScreenShotManager.h
+++ b/Engine/Source/Developer/ScreenShotComparisonTools/Private/ScreenShotManager.h
@@ -50,7 +50,8 @@ public:
 	virtual FScreenshotExportResult ExportScreenshotComparisonResult(FString ScreenshotName, FString ExportPath = TEXT("")) override;
 
 	virtual bool OpenComparisonReports(FString ImportPath, TArray<FComparisonReport>& OutReports) override;
-
+	
+	virtual FString GetIdealApprovedFolderForImage(const FAutomationScreenshotMetadata& MetaData) const override;
 	//~ End IScreenShotManager Interface
 
 private:
@@ -59,7 +60,6 @@ private:
 	FString	GetPathComponentForPlatformAndRHI(const FAutomationScreenshotMetadata& MetaData) const;
 	FString GetPathComponentForTestImages(const FAutomationScreenshotMetadata& MetaData) const;
 		
-	FString GetIdealApprovedFolderForImage(const FAutomationScreenshotMetadata& MetaData) const;
 
 	FString GetApprovedFolderForImageWithOptions(const FAutomationScreenshotMetadata& MetaData, EApprovedFolderOptions InOptions) const;
 
diff --git a/Engine/Source/Developer/ScreenShotComparisonTools/Public/Interfaces/IScreenShotManager.h b/Engine/Source/Developer/ScreenShotComparisonTools/Public/Interfaces/IScreenShotManager.h
--- a/Engine/Source/Developer/ScreenShotComparisonTools/Public/Interfaces/IScreenShotManager.h
+++ b/Engine/Source/Developer/ScreenShotComparisonTools/Public/Interfaces/IScreenShotManager.h
@@ -74,5 +74,10 @@ public:
 	 * Imports screenshot comparison data from a given path.
 	 */
 	virtual bool OpenComparisonReports(FString ImportPath, TArray<FComparisonReport>& OutReports) = 0;
-
+	
+	/**
+	* Calculate the ideal path for already already approved (ground truth) images. This handles abstracting away where Ground Truth images should be saved
+	* based on the platform/rhi/etc.
+	*/
+	virtual FString GetIdealApprovedFolderForImage(const FAutomationScreenshotMetadata& MetaData) const = 0;
 };
