#ue5 - Fix type conversion warnings in modules: AutomationController, AutomationTest, AVEncoder, AVIWriter.

diff --git a/Engine/Source/Developer/AutomationController/Private/AutomationControllerManager.h b/Engine/Source/Developer/AutomationController/Private/AutomationControllerManager.h
--- a/Engine/Source/Developer/AutomationController/Private/AutomationControllerManager.h
+++ b/Engine/Source/Developer/AutomationController/Private/AutomationControllerManager.h
@@ -573,7 +573,7 @@ private:
 		/** The test runners message address */
 		FMessageAddress OwnerMessageAddress;
 		/** The time since we had a ping from the instance*/
-		float LastPingTime;
+		double LastPingTime;
 	};
 
 	/** A array of running tests. */
diff --git a/Engine/Source/Developer/ScreenShotComparisonTools/Public/ImageComparer.h b/Engine/Source/Developer/ScreenShotComparisonTools/Public/ImageComparer.h
--- a/Engine/Source/Developer/ScreenShotComparisonTools/Public/ImageComparer.h
+++ b/Engine/Source/Developer/ScreenShotComparisonTools/Public/ImageComparer.h
@@ -109,8 +109,8 @@ public:
 	{
 		const bool AlphaSimilar = FMath::IsNearlyEqual((float)ColorA.A, ColorB.A, Tolerance.Alpha);
 
-		const float BrightnessA = FPixelOperations::GetLuminance(ColorA);
-		const float BrightnessB = FPixelOperations::GetLuminance(ColorB);
+		const double BrightnessA = FPixelOperations::GetLuminance(ColorA);
+		const double BrightnessB = FPixelOperations::GetLuminance(ColorB);
 		const bool BrightnessSimilar = FMath::IsNearlyEqual(BrightnessA, BrightnessB, Tolerance.MinBrightness);
 
 		return BrightnessSimilar && AlphaSimilar;
@@ -135,8 +135,8 @@ public:
 
 	static FORCEINLINE bool IsContrasting(const FColor& ColorA, const FColor& ColorB, const FImageTolerance& Tolerance)
 	{
-		const float BrightnessA = FPixelOperations::GetLuminance(ColorA);
-		const float BrightnessB = FPixelOperations::GetLuminance(ColorB);
+		const double BrightnessA = FPixelOperations::GetLuminance(ColorA);
+		const double BrightnessB = FPixelOperations::GetLuminance(ColorB);
 
 		return FMath::Abs(BrightnessA - BrightnessB) > Tolerance.MaxBrightness;
 	}
diff --git a/Engine/Source/Runtime/AVEncoder/Private/Decoders/VideoDecoderH264_Dummy.cpp b/Engine/Source/Runtime/AVEncoder/Private/Decoders/VideoDecoderH264_Dummy.cpp
--- a/Engine/Source/Runtime/AVEncoder/Private/Decoders/VideoDecoderH264_Dummy.cpp
+++ b/Engine/Source/Runtime/AVEncoder/Private/Decoders/VideoDecoderH264_Dummy.cpp
@@ -201,7 +201,7 @@ FVideoDecoder::EDecodeResult FVideoDecoderH264_Dummy::Decode(const FVideoDecoder
 			int C = ++YOffset;
 			for(int32 y=0,yMax=Height; y<yMax; ++y)
 			{
-				FMemory::Memset(pY, C++, Width);
+				FMemory::Memset(pY, (uint8)C++, Width);
 				pY += Width;
 			}
 			OnDecodedFrame(pNew.Release());
diff --git a/Engine/Source/Runtime/AVEncoder/Private/Decoders/vdecmpeg4/M4MemOps_Generic.cpp b/Engine/Source/Runtime/AVEncoder/Private/Decoders/vdecmpeg4/M4MemOps_Generic.cpp
--- a/Engine/Source/Runtime/AVEncoder/Private/Decoders/vdecmpeg4/M4MemOps_Generic.cpp
+++ b/Engine/Source/Runtime/AVEncoder/Private/Decoders/vdecmpeg4/M4MemOps_Generic.cpp
@@ -35,7 +35,7 @@ struct CLIPinitializer
 		iclp = iclip + 512;
 		for(int32 i= -512; i<512; i++)
 		{
-			iclp[i] = i < -256 ? -256 : i > 255 ? 255 : i;
+			iclp[i] = IntCastChecked<int16>(i < -256 ? -256 : i > 255 ? 255 : i);
 		}
 	}
 };
diff --git a/Engine/Source/Runtime/AVEncoder/Private/VideoDecoderUtilities.h b/Engine/Source/Runtime/AVEncoder/Private/VideoDecoderUtilities.h
--- a/Engine/Source/Runtime/AVEncoder/Private/VideoDecoderUtilities.h
+++ b/Engine/Source/Runtime/AVEncoder/Private/VideoDecoderUtilities.h
@@ -77,7 +77,7 @@ public:
 	 */
 	bool WaitTimeout(int64 InMicroSeconds)
 	{
-		return Event->Wait(FTimespan::FromMicroseconds(InMicroSeconds));
+		return Event->Wait(FTimespan::FromMicroseconds((double)InMicroSeconds));
 	}
 
 	/**
@@ -88,7 +88,7 @@ public:
 	bool WaitTimeoutAndReset(int64 InMicroSeconds)
 	{
 		check(InMicroSeconds > 0);	// 0 is forbidden. either call Wait() without timeout, or check the signal using IsSignaled()
-		if (Event->Wait(FTimespan::FromMicroseconds(InMicroSeconds)))
+		if (Event->Wait(FTimespan::FromMicroseconds((double)InMicroSeconds)))
 		{
 			Event->Reset();
 			return true;
diff --git a/Engine/Source/Runtime/AVIWriter/Private/AVIWriter.cpp b/Engine/Source/Runtime/AVIWriter/Private/AVIWriter.cpp
--- a/Engine/Source/Runtime/AVIWriter/Private/AVIWriter.cpp
+++ b/Engine/Source/Runtime/AVIWriter/Private/AVIWriter.cpp
@@ -377,7 +377,7 @@ public:
 			FString Directory = Options.OutputFilename;
 			FString Ext = FPaths::GetExtension(Directory, true);
 
-			int32 FPS = FMath::RoundToInt(double(Options.CaptureFramerateNumerator) / Options.CaptureFramerateDenominator);
+			int32 FPS = FMath::RoundToInt32(double(Options.CaptureFramerateNumerator) / Options.CaptureFramerateDenominator);
 
 			// Keep 3 seconds worth of frames in memory
 			CapturedFrames.Reset(new FCapturedFrames(Directory.LeftChop(Ext.Len()) + TEXT("_tmp"), FPS * 3));
diff --git a/Engine/Source/Runtime/AVIWriter/Private/CapturePin.cpp b/Engine/Source/Runtime/AVIWriter/Private/CapturePin.cpp
--- a/Engine/Source/Runtime/AVIWriter/Private/CapturePin.cpp
+++ b/Engine/Source/Runtime/AVIWriter/Private/CapturePin.cpp
@@ -12,7 +12,7 @@ DEFINE_LOG_CATEGORY(LogMovieCapture);
 
 FCapturePin::FCapturePin(HRESULT *phr, CSource *pFilter, const FAVIWriter& InWriter)
         : CSourceStream(NAME("Push Source"), phr, pFilter, L"Capture")
-		, FrameLength(UNITS / (InWriter.Options.CaptureFramerateNumerator / (double)InWriter.Options.CaptureFramerateDenominator))
+		, FrameLength(FMath::TruncToInt64(UNITS / (InWriter.Options.CaptureFramerateNumerator / (double)InWriter.Options.CaptureFramerateDenominator)))
 		, Writer(InWriter)
 {
 	ImageWidth  = Writer.GetWidth();
@@ -260,8 +260,8 @@ HRESULT FCapturePin::FillBuffer(IMediaSample *pSample)
 	// The current time is the sample's start.
 
 	// Not strictly necessary since AVI is constant framerate
-	REFERENCE_TIME StartTime = UNITS*CurrentFrame->StartTimeSeconds;
-	REFERENCE_TIME StopTime = UNITS*CurrentFrame->EndTimeSeconds;
+	REFERENCE_TIME StartTime = FMath::TruncToInt64(UNITS * CurrentFrame->StartTimeSeconds);
+	REFERENCE_TIME StopTime = FMath::TruncToInt64(UNITS * CurrentFrame->EndTimeSeconds);
 	pSample->SetTime(&StartTime, &StopTime);
 
 	REFERENCE_TIME StartMediaTime = CurrentFrame->FrameIndex;
diff --git a/Engine/Source/Runtime/AutomationTest/Public/AutomationTestExcludelist.h b/Engine/Source/Runtime/AutomationTest/Public/AutomationTestExcludelist.h
--- a/Engine/Source/Runtime/AutomationTest/Public/AutomationTestExcludelist.h
+++ b/Engine/Source/Runtime/AutomationTest/Public/AutomationTestExcludelist.h
@@ -73,7 +73,7 @@ struct FAutomationTestExcludelistEntry
 		int32 Num_Flags = Enum->NumEnums();
 		for (int32 i = 0; i < Num_Flags; i++)
 		{
-			int8 Flag = Enum->GetValueByIndex(i);
+			int8 Flag = IntCastChecked<int8>(Enum->GetValueByIndex(i));
 			if ((int8)ERHI_Flags::NUM == Flag)
 				break; // We reach the maximum value
 
diff --git a/Engine/Source/Runtime/Core/Public/Misc/AutomationTest.h b/Engine/Source/Runtime/Core/Public/Misc/AutomationTest.h
--- a/Engine/Source/Runtime/Core/Public/Misc/AutomationTest.h
+++ b/Engine/Source/Runtime/Core/Public/Misc/AutomationTest.h
@@ -301,7 +301,7 @@ public:
 	 *
 	 * @Param InTestFlags - the child test flag to add.
 	 */
-	void AddTestFlags( const uint8 InTestFlags)
+	void AddTestFlags( const uint32 InTestFlags)
 	{
 		TestFlags |= InTestFlags;
 	}
diff --git a/Engine/Source/Runtime/Core/Public/Templates/UnrealTemplate.h b/Engine/Source/Runtime/Core/Public/Templates/UnrealTemplate.h
--- a/Engine/Source/Runtime/Core/Public/Templates/UnrealTemplate.h
+++ b/Engine/Source/Runtime/Core/Public/Templates/UnrealTemplate.h
@@ -191,7 +191,7 @@ constexpr bool FloatFitsIn(InType In, InType Precision)
 	static_assert(std::is_floating_point_v<InType> && std::is_floating_point_v<OutType>, "Only floating point supported");
 	
 	OutType Out = static_cast<OutType>(In);
-	return std::abs(static_cast<InType>(Out) - In) <= Precision;
+	return fabs(static_cast<InType>(Out) - In) <= Precision;
 }
 
 template<typename OutType, typename InType>
diff --git a/Engine/Source/Runtime/Engine/Public/CanvasTypes.h b/Engine/Source/Runtime/Engine/Public/CanvasTypes.h
--- a/Engine/Source/Runtime/Engine/Public/CanvasTypes.h
+++ b/Engine/Source/Runtime/Engine/Public/CanvasTypes.h
@@ -659,10 +659,10 @@ public:
 	}
 
 	UE_DEPRECATED(5.0, "Use FCanvas::GetTime()")
-	float GetCurrentRealTime() const { return GetTime().GetRealTimeSeconds(); }
+	float GetCurrentRealTime() const { return FloatCastChecked<float>(GetTime().GetRealTimeSeconds(), UE_DOUBLE_SMALL_NUMBER); }
 
 	UE_DEPRECATED(5.0, "Use FCanvas::GetTime()")
-	float GetCurrentWorldTime() const { return GetTime().GetWorldTimeSeconds(); }
+	float GetCurrentWorldTime() const { return FloatCastChecked<float>(GetTime().GetWorldTimeSeconds(), UE_DOUBLE_SMALL_NUMBER); }
 
 	UE_DEPRECATED(5.0, "Use FCanvas::GetTime()")
 	float GetCurrentDeltaWorldTime() const { return GetTime().GetDeltaWorldTimeSeconds(); }
@@ -1310,7 +1310,7 @@ struct FScreenMessageWriter
 
 	void DrawLine(const FText& Message, int32 X = 10, const FLinearColor& Color = FLinearColor(1.0, 0.05, 0.05, 1.0))
 	{
-		Canvas.DrawShadowedText(X, Y, Message, GetStatsFont(), Color);
+		Canvas.DrawShadowedText((float)X, (float)Y, Message, GetStatsFont(), Color);
 		EmptyLine();
 	}
 
diff --git a/Engine/Source/Runtime/SignalProcessing/Public/SampleBuffer.h b/Engine/Source/Runtime/SignalProcessing/Public/SampleBuffer.h
--- a/Engine/Source/Runtime/SignalProcessing/Public/SampleBuffer.h
+++ b/Engine/Source/Runtime/SignalProcessing/Public/SampleBuffer.h
@@ -179,7 +179,7 @@ namespace Audio
 				// Convert from int to float:
 				for (int32 SampleIndex = 0; SampleIndex < NumSamples; SampleIndex++)
 				{
-					RawPCMData[SampleIndex] = ((float)Other.RawPCMData[SampleIndex]) / 32767.0f;
+					RawPCMData[SampleIndex] = (float)((float)Other.RawPCMData[SampleIndex]) / 32767.0f;
 				}
 			}
 			else
