Fix screenshots path names in test reports (standalone/no editor mode) (Contributed by [REDACTED_PERSON])
PR #9680

diff --git a/Engine/Source/Developer/ScreenShotComparisonTools/Private/ScreenShotManager.cpp b/Engine/Source/Developer/ScreenShotComparisonTools/Private/ScreenShotManager.cpp
--- a/Engine/Source/Developer/ScreenShotComparisonTools/Private/ScreenShotManager.cpp
+++ b/Engine/Source/Developer/ScreenShotComparisonTools/Private/ScreenShotManager.cpp
@@ -560,15 +560,14 @@ void FScreenShotManager::CopyDirectory(const FString& DestDir, const FString& Sr
 {
 	TArray<FString> FilesToCopy;
 
-	FString NormalizedSrc = SrcDir;
-	FPaths::NormalizeDirectoryName(NormalizedSrc);
+	FString AbsoluteSrcDir = FPaths::ConvertRelativePathToFull(SrcDir);
 
-	IFileManager::Get().FindFilesRecursive(FilesToCopy, *SrcDir, TEXT("*"), /*Files=*/true, /*Directories=*/false);
+	IFileManager::Get().FindFilesRecursive(FilesToCopy, *AbsoluteSrcDir, TEXT("*"), /*Files=*/true, /*Directories=*/false);
 
 	ParallelFor(FilesToCopy.Num(), [&](int32 Index)
 		{
 			const FString& SourceFilePath = FilesToCopy[Index];
-			FString DestFilePath = FPaths::Combine(DestDir, SourceFilePath.RightChop(SrcDir.Len()));
+			FString DestFilePath = FPaths::Combine(DestDir, SourceFilePath.RightChop(AbsoluteSrcDir.Len()));
 			IFileManager::Get().Copy(*DestFilePath, *SourceFilePath, true, true);
 		});
 }
