Improve ImageComparer by extending FComparableImage to be more straightforward when comparing images that are only in-memory.

diff --git a/Engine/Source/Developer/ScreenShotComparisonTools/Private/ImageComparer.cpp b/Engine/Source/Developer/ScreenShotComparisonTools/Private/ImageComparer.cpp
--- a/Engine/Source/Developer/ScreenShotComparisonTools/Private/ImageComparer.cpp
+++ b/Engine/Source/Developer/ScreenShotComparisonTools/Private/ImageComparer.cpp
@@ -215,10 +215,24 @@ FComparisonReport::FComparisonReport(const FString& InReportRootDirectory, const
 	ReportPath = FPaths::GetPath(InReportFile);
 }
 
-TSharedPtr<FComparableImage> FImageComparer::Open(const FString& ImagePath, FText& OutError)
+bool FComparableImage::LoadFile(const FString& ImagePath, FText& OutError)
 {
-	LLM_SCOPE_BYNAME(TEXT("AutomationTest/ImageCompare"));
+	TArray64<uint8> ImageData;
+	const bool OpenSuccess = FFileHelper::LoadFileToArray(ImageData, *ImagePath);
+
+	if ( !OpenSuccess )
+	{
+		OutError = LOCTEXT("ErrorOpeningImageA", "Unable to read image");
+		return false;
+	}
+
 	const FString ImageExtension = FPaths::GetExtension(ImagePath);
+	return LoadCompressedData(ImageData.GetData(), ImageData.Num(), ImageExtension, OutError);
+}
+
+bool FComparableImage::LoadCompressedData(const void* CompressedData, int64 CompressedSize, const FString& ImageExtension, FText& OutError)
+{
+	LLM_SCOPE_BYNAME(TEXT("AutomationTest/ImageCompare"));
 	const EImageFormat ImageFormat = ImageWrapperHelper::GetImageFormat(ImageExtension);
 
 	IImageWrapperModule& ImageWrapperModule = FModuleManager::GetModuleChecked<IImageWrapperModule>(FName("ImageWrapper"));
@@ -227,47 +241,32 @@ TSharedPtr<FComparableImage> FImageComparer::Open(const FString& ImagePath, FTex
 	if ( !ImageReader.IsValid() )
 	{
 		OutError = FText::Format(LOCTEXT("ImageWrapperMissing", "Unable to locate image processor for file format {0}"), FText::FromString(ImageExtension));
-		return nullptr;
+		return false;
 	}
 
-	TArray64<uint8> PngData;
-	const bool OpenSuccess = FFileHelper::LoadFileToArray(PngData, *ImagePath);
-
-	if ( !OpenSuccess )
-	{
-		OutError = LOCTEXT("ErrorOpeningImageA", "Unable to read image");
-		return nullptr;
-	}
-
-	if ( !ImageReader->SetCompressed(PngData.GetData(), PngData.Num()) )
+	if ( !ImageReader->SetCompressed(CompressedData, CompressedSize) )
 	{
 		OutError = LOCTEXT("ErrorParsingImageA", "Unable to parse image");
-		return nullptr;
+		return false;
 	}
 
-	TSharedPtr<FComparableImage> Image = MakeShareable(new FComparableImage());
-	
-	if ( !ImageReader->GetRaw(ERGBFormat::RGBA, 8, Image->Bytes) )
+	if ( !ImageReader->GetRaw(ERGBFormat::RGBA, 8, Bytes) )
 	{
 		OutError = LOCTEXT("ErrorReadingRawDataA", "Unable to decompress image");
-		return nullptr;
-	}
-	else
-	{
-		Image->Width = ImageReader->GetWidth();
-		Image->Height = ImageReader->GetHeight();
+		return false;
 	}
 
-	return Image;
+	Width = ImageReader->GetWidth();
+	Height = ImageReader->GetHeight();
+	return true;
 }
 
 FImageComparisonResult FImageComparer::Compare(const FString& ImagePathA, const FString& ImagePathB, FImageTolerance Tolerance, const FString& OutDeltaPath)
 {
 	FText ErrorA, ErrorB;
-	TSharedPtr<FComparableImage> ImageA = Open(ImagePathA, ErrorA);
-	TSharedPtr<FComparableImage> ImageB = Open(ImagePathB, ErrorB);
+	FComparableImage ImageA, ImageB;
 
-	if ( !ImageA.IsValid() )
+	if ( !ImageA.LoadFile(ImagePathA, ErrorA) )
 	{
 		FImageComparisonResult Results;
 		Results.ApprovedFilePath = ImagePathA;
@@ -276,7 +275,7 @@ FImageComparisonResult FImageComparer::Compare(const FString& ImagePathA, const
 		return Results;
 	}
 
-	if ( !ImageB.IsValid() )
+	if ( !ImageB.LoadFile(ImagePathB, ErrorB) )
 	{
 		FImageComparisonResult Results;
 		Results.ApprovedFilePath = ImagePathA;
@@ -285,7 +284,7 @@ FImageComparisonResult FImageComparer::Compare(const FString& ImagePathA, const
 		return Results;
 	}
 
-	FImageComparisonResult Results = Compare(ImageA.Get(), ImageB.Get(), Tolerance, OutDeltaPath);
+	FImageComparisonResult Results = Compare(&ImageA, &ImageB, Tolerance, OutDeltaPath);
 	Results.ApprovedFilePath = ImagePathA;
 	Results.IncomingFilePath = ImagePathB;
 	return Results;
@@ -420,22 +419,21 @@ double FImageComparer::CompareStructuralSimilarity(const FString& ImagePathA, co
 	Results.IncomingFilePath = ImagePathB;
 
 	FText ErrorA, ErrorB;
-	TSharedPtr<FComparableImage> ImageA = Open(ImagePathA, ErrorA);
-	TSharedPtr<FComparableImage> ImageB = Open(ImagePathB, ErrorB);
+	FComparableImage ImageA, ImageB;
 
-	if ( !ImageA.IsValid() )
+	if ( !ImageA.LoadFile(ImagePathA, ErrorA) )
 	{
 		Results.ErrorMessage = ErrorA;
 		return 0.0f;
 	}
 
-	if ( !ImageB.IsValid() )
+	if ( !ImageB.LoadFile(ImagePathB, ErrorB) )
 	{
 		Results.ErrorMessage = ErrorB;
 		return 0.0f;
 	}
 
-	if ( ImageA->Width != ImageB->Width || ImageA->Height != ImageB->Height )
+	if ( ImageA.Width != ImageB.Width || ImageA.Height != ImageB.Height )
 	{
 		Results.ErrorMessage = LOCTEXT("DifferentSizesUnsupported", "We can not compare images of different sizes at this time.");
 		return 0.0f;
@@ -444,7 +442,7 @@ double FImageComparer::CompareStructuralSimilarity(const FString& ImagePathA, co
 	//ImageA->Process();
 	//ImageB->Process();
 
-	return CompareStructuralSimilarity(ImageA.Get(), ImageB.Get(), InCompareComponent, OutDeltaPath);
+	return CompareStructuralSimilarity(&ImageA, &ImageB, InCompareComponent, OutDeltaPath);
 }
 
 double FImageComparer::CompareStructuralSimilarity(const FComparableImage* ImageA, const FComparableImage* ImageB, EStructuralSimilarityComponent InCompareComponent, const FString& OutDeltaPath)
diff --git a/Engine/Source/Developer/ScreenShotComparisonTools/Public/ImageComparer.h b/Engine/Source/Developer/ScreenShotComparisonTools/Public/ImageComparer.h
--- a/Engine/Source/Developer/ScreenShotComparisonTools/Public/ImageComparer.h
+++ b/Engine/Source/Developer/ScreenShotComparisonTools/Public/ImageComparer.h
@@ -149,7 +149,7 @@ public:
 /**
  *
  */
-class FComparableImage
+class SCREENSHOTCOMPARISONTOOLS_API FComparableImage
 {
 public:
 	int32 Width = 0;
@@ -176,6 +176,28 @@ public:
 			Bytes[Offset + 2],
 			Bytes[Offset + 3]);
 	}
+
+	/**
+	 * Populate image by loading an file
+	 *
+	 * @param ImagePath Path for the image file to load
+	 * @param OutError Contains the error message if load fails
+	 * 
+	 * @return true if success
+	*/
+	bool LoadFile(const FString& ImagePath, FText& OutError);
+
+	/**
+	 * Populate image by loading compressed data
+	 *
+	 * @param CompressedData The memory address of the start of the compressed data.
+	 * @param CompressedSize The size of the compressed data parsed.
+	 * @param ImageExtension File extension of the image format
+	 * @param OutError Contains the error message if load fails
+	 * 
+	 * @return true if success
+	*/
+	bool LoadCompressedData(const void* CompressedData, int64 CompressedSize, const FString& ImageExtension, FText& OutError);
 };
 
 /**
@@ -429,7 +451,4 @@ public:
 	 */
 	double CompareStructuralSimilarity(const FString& ImagePathA, const FString& ImagePathB, EStructuralSimilarityComponent InCompareComponent, const FString& OutDeltaPath);
 	double CompareStructuralSimilarity(const FComparableImage* ImageA, const FComparableImage* ImageB, EStructuralSimilarityComponent InCompareComponent, const FString& OutDeltaPath);
-
-private:
-	TSharedPtr<FComparableImage> Open(const FString& ImagePath, FText& OutError);
 };
