Fix for not found screenshot reference while running automated tests in UFE

diff --git a/Engine/Source/Developer/ScreenShotComparisonTools/Private/ScreenShotManager.cpp b/Engine/Source/Developer/ScreenShotComparisonTools/Private/ScreenShotManager.cpp
--- a/Engine/Source/Developer/ScreenShotComparisonTools/Private/ScreenShotManager.cpp
+++ b/Engine/Source/Developer/ScreenShotComparisonTools/Private/ScreenShotManager.cpp
@@ -275,8 +275,6 @@ FImageComparisonResult FScreenShotManager::CompareScreenshot(const FString& InUn
 		}
 
 		FString NearestExistingApprovedImage;
-		
-
 		if ( ExistingMetadata.IsSet() )
 		{
 			int32 MatchScore = -1;
@@ -365,8 +363,6 @@ FImageComparisonResult FScreenShotManager::CompareScreenshot(const FString& InUn
 	// to the data remaining intact
 	FString ReportPathOnDisk = FPaths::Combine(ScreenshotResultsFolder, ResultsSubFolder, IncomingMetaData.Platform, GetPathComponentForRHI(IncomingMetaData), TEXT("/"));
 
-	FString ProjectDir = FPaths::GetPath(FPaths::GetProjectFilePath());
-
 	/*
 		Now create copies of all three images for the report. We use copies (a move in the case of the delta image) so that the report folders
 		can be moved around without any external references. This is vital for CIS where we want the results to be served up via HTTP or just 
@@ -386,6 +382,8 @@ FImageComparisonResult FScreenShotManager::CompareScreenshot(const FString& InUn
 		ComparisonResult.ReportApprovedFilePath = FPaths::Combine(ReportPathOnDisk, ApprovedName);
 	}
 
+	const FString ProjectDir = FPaths::GetPath(FPaths::IsProjectFilePathSet() ? FPaths::GetProjectFilePath() : FPaths::ProjectDir());
+
 	/*
 		First we need the incoming file in the report. If the calling code wants to keep the image then do a copy, 
 		otherwise move it
@@ -396,7 +394,7 @@ FImageComparisonResult FScreenShotManager::CompareScreenshot(const FString& InUn
 //	if (CanMoveImage == false)
 	{
 		// copy the incoming file to the report path 
-		FString IncomingFileFullPath = *FPaths::Combine(ProjectDir, ComparisonResult.IncomingFilePath);
+		const FString IncomingFileFullPath = *FPaths::Combine(ProjectDir, ComparisonResult.IncomingFilePath);
 
 		if (IFileManager::Get().Copy(*ComparisonResult.ReportIncomingFilePath, *IncomingFileFullPath, true, true) == COPY_OK)
 		{
