Fix: UnrealFrontend - Screenshot reference can't be found in Session frontend

diff --git a/Engine/Source/Developer/ScreenShotComparisonTools/Private/ScreenShotManager.cpp b/Engine/Source/Developer/ScreenShotComparisonTools/Private/ScreenShotManager.cpp
--- a/Engine/Source/Developer/ScreenShotComparisonTools/Private/ScreenShotManager.cpp
+++ b/Engine/Source/Developer/ScreenShotComparisonTools/Private/ScreenShotManager.cpp
@@ -17,6 +17,7 @@
 #include "PlatformInfo.h"
 #include "Misc/DataDrivenPlatformInfoRegistry.h"
 #include "ScreenShotComparisonSettings.h"
+#include "ILauncherServicesModule.h"
 
 
 DEFINE_LOG_CATEGORY(LogScreenShotManager);
@@ -96,6 +97,15 @@ FString FScreenShotManager::GetApprovedFolderForImageWithOptions(const FAutomati
 
 	FString OutPath = FPaths::ProjectDir();
 
+	// Project path would be different if project is started from Unreal Frontend (uses LauncherServices)
+	if (FModuleManager::Get().IsModuleLoaded("LauncherServices"))
+	{
+		ILauncherServicesModule& LauncherServicesModule = FModuleManager::LoadModuleChecked<ILauncherServicesModule>("LauncherServices");
+		FString ProjectPath = FPaths::GetPath(LauncherServicesModule.GetProfileManager()->GetProjectPath());
+		FPaths::MakePathRelativeTo(ProjectPath, *FPaths::ProjectDir());
+		OutPath = FPaths::DirectoryExists(ProjectPath) ? ProjectPath : OutPath;
+	}
+
 	if (bUsePlatformPath)
 	{
 		OutPath = FPaths::Combine(OutPath, TEXT("Platforms"), MetaData.Platform, TEXT("Test/Screenshots/"), TestFolder, GetPathComponentForRHI(MetaData));
diff --git a/Engine/Source/Developer/ScreenShotComparisonTools/ScreenShotComparisonTools.Build.cs b/Engine/Source/Developer/ScreenShotComparisonTools/ScreenShotComparisonTools.Build.cs
--- a/Engine/Source/Developer/ScreenShotComparisonTools/ScreenShotComparisonTools.Build.cs
+++ b/Engine/Source/Developer/ScreenShotComparisonTools/ScreenShotComparisonTools.Build.cs
@@ -27,6 +27,7 @@ public class ScreenShotComparisonTools : ModuleRules
 		PrivateIncludePathModuleNames.AddRange(
 			new string[] {
 				"MessagingCommon",
+				"LauncherServices"
 			}
 		);
 
