Support Confidential platforms in Screenshot fallback settings

#rnx

diff --git a/Engine/Source/Developer/ScreenShotComparisonTools/Private/ScreenShotManager.cpp b/Engine/Source/Developer/ScreenShotComparisonTools/Private/ScreenShotManager.cpp
--- a/Engine/Source/Developer/ScreenShotComparisonTools/Private/ScreenShotManager.cpp
+++ b/Engine/Source/Developer/ScreenShotComparisonTools/Private/ScreenShotManager.cpp
@@ -45,9 +45,8 @@ FScreenShotManager::FScreenShotManager()
 
 	// the automation controller owns the screenshot manager so we don't have to worry about outliving it
 	//AutomationController->OnTestsAvailable().AddRaw(this, &FScreenShotManager::OnTestAvailableCallback);
-	// #agrant todo - we don't want this dependency and higher level code should clear the folders before running new tests
 
-	BuildFallbackPlatformsListFromConfig(ScreenshotSettings);
+	BuildFallbackPlatformsListFromConfig();
 }
 
 FScreenShotManager::~FScreenShotManager()
@@ -582,13 +581,19 @@ void FScreenShotManager::CopyDirectory(const FString& DestDir, const FString& Sr
 		});
 }
 
-void FScreenShotManager::BuildFallbackPlatformsListFromConfig(const UScreenShotComparisonSettings* InSettings)
+void FScreenShotManager::BuildFallbackPlatformsListFromConfig()
 {
 	FallbackPlatforms.Empty();
 
-	if (InSettings->ScreenshotFallbackPlatforms.Num() > 0)
+	const UScreenShotComparisonSettings* ScreenshotSettings = GetDefault<UScreenShotComparisonSettings>();
+	TSet<FScreenshotFallbackEntry> ScreenshotFallbackPlatforms(ScreenshotSettings->ScreenshotFallbackPlatforms);
+#if WITH_EDITOR
+	ScreenshotFallbackPlatforms.Append(UScreenShotComparisonSettings::GetAllPlatformSettings());
+#endif // WITH_EDITOR
+
+	if (ScreenshotFallbackPlatforms.Num() > 0)
 	{
-		for (const FScreenshotFallbackEntry& Entry : InSettings->ScreenshotFallbackPlatforms)
+		for (const FScreenshotFallbackEntry& Entry : ScreenshotFallbackPlatforms)
 		{
 			FString Parent = Entry.Parent;
 			FString Child = Entry.Child;
diff --git a/Engine/Source/Developer/ScreenShotComparisonTools/Private/ScreenShotManager.h b/Engine/Source/Developer/ScreenShotComparisonTools/Private/ScreenShotManager.h
--- a/Engine/Source/Developer/ScreenShotComparisonTools/Private/ScreenShotManager.h
+++ b/Engine/Source/Developer/ScreenShotComparisonTools/Private/ScreenShotManager.h
@@ -71,7 +71,7 @@ private:
 
 	void CopyDirectory(const FString& DestDir, const FString& SrcDir, const FString& Pattern);
 
-	void BuildFallbackPlatformsListFromConfig(const UScreenShotComparisonSettings* ScreenShotSettings);
+	void BuildFallbackPlatformsListFromConfig();
 
 private:
 
diff --git a/Engine/Source/Developer/ScreenShotComparisonTools/Private/ScreenshotComparisonSettings.cpp b/Engine/Source/Developer/ScreenShotComparisonTools/Private/ScreenshotComparisonSettings.cpp
--- a/Engine/Source/Developer/ScreenShotComparisonTools/Private/ScreenshotComparisonSettings.cpp
+++ b/Engine/Source/Developer/ScreenShotComparisonTools/Private/ScreenshotComparisonSettings.cpp
@@ -3,9 +3,71 @@
 
 #include "ScreenShotComparisonSettings.h"
 
+#if WITH_EDITOR
+#include "Misc/CoreMisc.h"
+#include "Misc/DataDrivenPlatformInfoRegistry.h"
+#endif // WITH_EDITOR
 
 UScreenShotComparisonSettings::UScreenShotComparisonSettings(const FObjectInitializer& ObjectInitializer)
 	: Super(ObjectInitializer)
 {
 	//DefaultScreenshotResolution = FIntPoint(1920, 1080);
 }
+
+void UScreenShotComparisonSettings::LoadSettings()
+{
+	ScreenshotFallbackPlatforms.Empty();
+	UObject::LoadConfig(GetClass());
+}
+
+UScreenShotComparisonSettings* UScreenShotComparisonSettings::Create(const FString& PlatformName)
+{
+	UScreenShotComparisonSettings* Settings = NewObject<UScreenShotComparisonSettings>();
+	Settings->SetPlatform(PlatformName);
+	return Settings;
+}
+
+const TCHAR* UScreenShotComparisonSettings::GetConfigOverridePlatform() const
+{
+	return *Platform;
+}
+
+const FString& UScreenShotComparisonSettings::GetPlatformName() const
+{
+	return Platform;
+}
+
+void UScreenShotComparisonSettings::SetPlatform(const FString& PlatformName)
+{
+	Platform = PlatformName;
+	LoadSettings();
+}
+
+#if WITH_EDITOR
+const TSet<FScreenshotFallbackEntry>& UScreenShotComparisonSettings::GetAllPlatformSettings()
+{
+	static TOptional<TSet<FScreenshotFallbackEntry>> ScreenshotFallbackEntries;
+
+	if (!ScreenshotFallbackEntries.IsSet())
+	{
+		ScreenshotFallbackEntries.Emplace();
+
+		for (const auto& Pair : FDataDrivenPlatformInfoRegistry::GetAllPlatformInfos())
+		{
+			if (Pair.Value.bIsFakePlatform || Pair.Value.bEnabledForUse == false)
+			{
+				continue;
+			}
+
+			FString PlatformName = Pair.Key.ToString();
+			UScreenShotComparisonSettings* Settings (UScreenShotComparisonSettings::Create(PlatformName));
+			check(nullptr != Settings);
+
+			ScreenshotFallbackEntries->Append(Settings->ScreenshotFallbackPlatforms);
+		}
+	}
+
+	return ScreenshotFallbackEntries.GetValue();
+}
+
+#endif // WITH_EDITOR
diff --git a/Engine/Source/Developer/ScreenShotComparisonTools/Public/ScreenShotComparisonSettings.h b/Engine/Source/Developer/ScreenShotComparisonTools/Public/ScreenShotComparisonSettings.h
--- a/Engine/Source/Developer/ScreenShotComparisonTools/Public/ScreenShotComparisonSettings.h
+++ b/Engine/Source/Developer/ScreenShotComparisonTools/Public/ScreenShotComparisonSettings.h
@@ -21,8 +21,17 @@ struct FScreenshotFallbackEntry
 
 	UPROPERTY(Config)
 	FString Child;
+
+	bool operator==(const FScreenshotFallbackEntry& Other) const
+	{
+		return Child == Other.Child;
+	}
 };
 
+FORCEINLINE uint32 GetTypeHash(const FScreenshotFallbackEntry& Object)
+{
+	return GetTypeHash(Object.Child);
+}
 
 UCLASS(config = Engine, defaultconfig)
 class SCREENSHOTCOMPARISONTOOLS_API UScreenShotComparisonSettings : public UObject
@@ -32,7 +41,7 @@ class SCREENSHOTCOMPARISONTOOLS_API UScreenShotComparisonSettings : public UObje
 public:
 
 	/**
-	* if true, any checked-in test results for confidential platforms will be put under <ProjectDir>/Platforms/<Platform>/Test instead of <ProjectDir>/Test
+	* If true, any checked-in test results for confidential platforms will be put under <ProjectDir>/Platforms/<Platform>/Test instead of <ProjectDir>/Test
 	*/
 	UPROPERTY(Config)
 	bool bUseConfidentialPlatformPathsForSavedResults;
@@ -42,5 +51,41 @@ public:
 	*/
 	UPROPERTY(Config)
 	TArray<FScreenshotFallbackEntry>  ScreenshotFallbackPlatforms;
-};
 
+	/**
+	* Creates class instance
+	* @param PlatformName Reference to a string containing platform name (if it is empty the current platform name is used).
+	*/
+	static UScreenShotComparisonSettings* Create(const FString& PlatformName = FString{});
+
+	/**
+	 * Loads settings of corresponding config.
+	 */
+	virtual void LoadSettings() final;
+
+	/**
+	 * Overrides config hierarchy platform to be used in UObject internals
+	 */
+	virtual const TCHAR* GetConfigOverridePlatform() const override;
+
+#if WITH_EDITOR
+public:
+	static const TSet<FScreenshotFallbackEntry>& GetAllPlatformSettings();
+#endif // WITH_EDITOR
+
+protected:
+
+	/**
+	 * Returns platform name reference. As the class can store platform-independent config, it returns an empty string if the platform was not specified.
+	 */
+	virtual const FString& GetPlatformName() const;
+
+	/**
+	 * Sets platform and reloads settings.
+	 * @param PlatformName Reference to a string containing platform name (if it is empty the default config is used).
+	 */
+	virtual void SetPlatform(const FString& PlatformName);
+
+private:
+	FString Platform;
+};
