Fixing EngineTest screenshot fallbacks for non-D3D platforms where the fallback paths aren't of the form RHI_FeatureLevel.

diff --git a/Engine/Source/Developer/ScreenShotComparisonTools/Private/ScreenShotManager.cpp b/Engine/Source/Developer/ScreenShotComparisonTools/Private/ScreenShotManager.cpp
--- a/Engine/Source/Developer/ScreenShotComparisonTools/Private/ScreenShotManager.cpp
+++ b/Engine/Source/Developer/ScreenShotComparisonTools/Private/ScreenShotManager.cpp
@@ -187,8 +187,17 @@ TArray<FString> FScreenShotManager::FindApprovedImages(const FAutomationScreensh
 
 				FAutomationScreenshotMetadata CopiedMetaData = IncomingMetaData;
 				CopiedMetaData.Platform = Components[0];
-				CopiedMetaData.Rhi = FeatureLevels[0];
-				CopiedMetaData.FeatureLevel = FeatureLevels[1];
+				if (FeatureLevels.Num() > 1)
+				{
+					CopiedMetaData.Rhi = FeatureLevels[0];
+					CopiedMetaData.FeatureLevel = FeatureLevels[1];
+				}
+				else
+				{
+					// We don't need to do RHI_FL, just FL
+					CopiedMetaData.Rhi = TEXT("");
+					CopiedMetaData.FeatureLevel = FeatureLevels[0];
+				}
 
 				ApprovedPath = FPaths::GetPath(GetIdealApprovedFolderForImage(CopiedMetaData));
 				IFileManager::Get().FindFilesRecursive(ApprovedImages, *ApprovedPath, TEXT("*.png"), true, false);
