Adding logging of failed fallback paths to ScreenShot tests

diff --git a/Engine/Source/Developer/ScreenShotComparisonTools/Private/ScreenShotManager.cpp b/Engine/Source/Developer/ScreenShotComparisonTools/Private/ScreenShotManager.cpp
--- a/Engine/Source/Developer/ScreenShotComparisonTools/Private/ScreenShotManager.cpp
+++ b/Engine/Source/Developer/ScreenShotComparisonTools/Private/ScreenShotManager.cpp
@@ -130,19 +130,30 @@ FString FScreenShotManager::GetIdealApprovedFolderForImage(const FAutomationScre
 
 TArray<FString> FScreenShotManager::FindApprovedImages(const FAutomationScreenshotMetadata& IncomingMetaData)
 {
+	TArray<FString> TriedPaths;
+
+	auto FindImages = [&TriedPaths](TArray<FString>& OutApprovedImages, const FString& InPath)
+	{
+		IFileManager::Get().FindFilesRecursive(OutApprovedImages, *InPath, TEXT("*.png"), true, false);
+		TriedPaths.Emplace(InPath);
+	};
+
 	TArray<FString> ApprovedImages;
 
 	EApprovedFolderOptions Options = bUseConfidentialPlatformPaths ? EApprovedFolderOptions::None : EApprovedFolderOptions::UsePlatformFolders;
 
 	// check out standard path using whether confidential platforms are in a separate tree
 	FString ApprovedPath = GetApprovedFolderForImageWithOptions(IncomingMetaData, Options);
-	IFileManager::Get().FindFilesRecursive(ApprovedImages, *ApprovedPath, TEXT("*.png"), true, false);
+	FindImages(ApprovedImages, ApprovedPath);
+
+	// Make sure the first log line is of the first path tried, not the last fallback. The list of fallbacks will be printed if nothing is found.
+	const FString FirstApprovedPath = ApprovedPath;
 
 	// check again, but try legacy paths
 	if (!ApprovedImages.Num())
 	{
 		ApprovedPath = GetApprovedFolderForImageWithOptions(IncomingMetaData, Options | EApprovedFolderOptions::UseLegacyPaths);
-		IFileManager::Get().FindFilesRecursive(ApprovedImages, *ApprovedPath, TEXT("*.png"), true, false);
+		FindImages(ApprovedImages, ApprovedPath);
 	}
 
 	// if we're a blank and bUseConfidentialPlatformPaths, try without that
@@ -150,13 +161,13 @@ TArray<FString> FScreenShotManager::FindApprovedImages(const FAutomationScreensh
 	{
 		// check legacy paths.
 		ApprovedPath = FPaths::GetPath(GetApprovedFolderForImageWithOptions(IncomingMetaData, EApprovedFolderOptions::None));
-		IFileManager::Get().FindFilesRecursive(ApprovedImages, *ApprovedPath, TEXT("*.png"), true, false);
+		FindImages(ApprovedImages, ApprovedPath);
 
 		// check again, but try legacy paths
 		if (!ApprovedImages.Num())
 		{
 			ApprovedPath = GetApprovedFolderForImageWithOptions(IncomingMetaData, EApprovedFolderOptions::UseLegacyPaths);
-			IFileManager::Get().FindFilesRecursive(ApprovedImages, *ApprovedPath, TEXT("*.png"), true, false);
+			FindImages(ApprovedImages, ApprovedPath);
 		}
 	}
 
@@ -165,7 +176,7 @@ TArray<FString> FScreenShotManager::FindApprovedImages(const FAutomationScreensh
 	{
 		FString CurrentPlatformRHI = GetPathComponentForPlatformAndRHI(IncomingMetaData);
 
-		UE_LOG(LogScreenShotManager, Log, TEXT("No ideal-image found at %s. Checking fallback images"), *ApprovedPath);
+		UE_LOG(LogScreenShotManager, Log, TEXT("No ideal-image found at %s. Checking fallback images"), *FirstApprovedPath);
 
 		while (ApprovedImages.Num() == 0)
 		{
@@ -200,13 +211,13 @@ TArray<FString> FScreenShotManager::FindApprovedImages(const FAutomationScreensh
 				}
 
 				ApprovedPath = FPaths::GetPath(GetIdealApprovedFolderForImage(CopiedMetaData));
-				IFileManager::Get().FindFilesRecursive(ApprovedImages, *ApprovedPath, TEXT("*.png"), true, false);
+				FindImages(ApprovedImages, ApprovedPath);
 
 				// check again, but try legacy paths
 				if (!ApprovedImages.Num())
 				{
 					ApprovedPath = GetApprovedFolderForImageWithOptions(CopiedMetaData, EApprovedFolderOptions::UseLegacyPaths);
-					IFileManager::Get().FindFilesRecursive(ApprovedImages, *ApprovedPath, TEXT("*.png"), true, false);
+					FindImages(ApprovedImages, ApprovedPath);
 				}
 
 				if (ApprovedImages.Num())
@@ -221,6 +232,15 @@ TArray<FString> FScreenShotManager::FindApprovedImages(const FAutomationScreensh
 		}
 	}
 
+	if (ApprovedImages.IsEmpty())
+	{
+		UE_LOG(LogScreenShotManager, Log, TEXT("Couldn't Find any fallback images, tried the following paths:"));
+		for (const FString& TriedPath : TriedPaths)
+		{
+			UE_LOG(LogScreenShotManager, Log, TEXT("    %s"), *TriedPath);
+		}
+	}
+
 	return ApprovedImages;
 }
 
