Fix fails to find screenshot reference when Project launcher is opened and the current project selected

#rnx

diff --git a/Engine/Source/Developer/ScreenShotComparisonTools/Private/ScreenShotManager.cpp b/Engine/Source/Developer/ScreenShotComparisonTools/Private/ScreenShotManager.cpp
--- a/Engine/Source/Developer/ScreenShotComparisonTools/Private/ScreenShotManager.cpp
+++ b/Engine/Source/Developer/ScreenShotComparisonTools/Private/ScreenShotManager.cpp
@@ -96,14 +96,15 @@ FString FScreenShotManager::GetApprovedFolderForImageWithOptions(const FAutomati
 
 	FString OutPath = FPaths::ProjectDir();
 
+#if !WITH_EDITOR
 	// Project path would be different if project is started from Unreal Frontend (uses LauncherServices)
 	if (FModuleManager::Get().IsModuleLoaded("LauncherServices"))
 	{
 		ILauncherServicesModule& LauncherServicesModule = FModuleManager::LoadModuleChecked<ILauncherServicesModule>("LauncherServices");
 		FString ProjectPath = FPaths::GetPath(LauncherServicesModule.GetProfileManager()->GetProjectPath());
-		FPaths::MakePathRelativeTo(ProjectPath, *FPaths::ProjectDir());
-		OutPath = FPaths::DirectoryExists(ProjectPath) ? ProjectPath : OutPath;
+		OutPath = !ProjectPath.IsEmpty() && FPaths::DirectoryExists(ProjectPath) ? ProjectPath : OutPath;
 	}
+#endif
 
 	if (bUsePlatformPath)
 	{
