Add frame comparison failure to the test report when comparing image sequence for Movie Pipeline Functional Test
+ Improve messaging when movie pipeline test is failing for reason other than image comparison.
+ Store ScreenshotName inside comparison request and result.

#rnx

diff --git a/Engine/Plugins/MovieScene/MovieRenderPipeline/Source/MovieRenderPipelineCore/Public/MovieRenderPipelineDataTypes.h b/Engine/Plugins/MovieScene/MovieRenderPipeline/Source/MovieRenderPipelineCore/Public/MovieRenderPipelineDataTypes.h
--- a/Engine/Plugins/MovieScene/MovieRenderPipeline/Source/MovieRenderPipelineCore/Public/MovieRenderPipelineDataTypes.h
+++ b/Engine/Plugins/MovieScene/MovieRenderPipeline/Source/MovieRenderPipelineCore/Public/MovieRenderPipelineDataTypes.h
@@ -1251,7 +1251,7 @@ struct FMoviePipelineShotOutputData
 	GENERATED_BODY()
 
 	/** Which shot was this output data for? */
-	UPROPERTY(VisibleAnywhere, BlueprintReadOnly, Category = "Movie Pipeline")
+	UPROPERTY(VisibleAnywhere, BlueprintReadOnly, Transient, Category = "Movie Pipeline")
 	TWeakObjectPtr<UMoviePipelineExecutorShot> Shot;
 
 	/** 
@@ -1310,11 +1310,11 @@ struct FMoviePipelineOutputData
 	*
 	* Provided here for backwards compatibility.
 	*/
-	UPROPERTY(VisibleAnywhere, BlueprintReadOnly, Category = "Movie Pipeline")
+	UPROPERTY(VisibleAnywhere, BlueprintReadOnly, Transient, Category = "Movie Pipeline")
 	TObjectPtr<UMoviePipeline> Pipeline;
 	
 	/** Job the data is for. Job may still be in progress (if a shot callback) so be careful about modifying properties on it */
-	UPROPERTY(VisibleAnywhere, BlueprintReadOnly, Category = "Movie Pipeline")
+	UPROPERTY(VisibleAnywhere, BlueprintReadOnly, Transient, Category = "Movie Pipeline")
 	TObjectPtr<UMoviePipelineExecutorJob> Job;
 	
 	/** Did the job succeed, or was it canceled early due to an error (such as failure to write file to disk)? */
diff --git a/Engine/Plugins/MovieScene/MovieRenderPipeline/Source/MovieRenderPipelineEditor/Private/MoviePipelineFunctionalTestBase.cpp b/Engine/Plugins/MovieScene/MovieRenderPipeline/Source/MovieRenderPipelineEditor/Private/MoviePipelineFunctionalTestBase.cpp
--- a/Engine/Plugins/MovieScene/MovieRenderPipeline/Source/MovieRenderPipelineEditor/Private/MoviePipelineFunctionalTestBase.cpp
+++ b/Engine/Plugins/MovieScene/MovieRenderPipeline/Source/MovieRenderPipelineEditor/Private/MoviePipelineFunctionalTestBase.cpp
@@ -162,22 +162,20 @@ void AMoviePipelineFunctionalTestBase::CompareRenderOutputToGroundTruth(FMoviePi
 	IScreenShotManagerPtr ScreenshotManager = ScreenShotModule.GetScreenShotManager();
 
 	// Now we know where to look for our ground truth data.
-	FString ReportDirectory = ScreenshotManager->GetIdealApprovedFolderForImage(MetaData);
-	FString GroundTruthFilename = ReportDirectory / "GroundTruth.json";
-
-	if (FPaths::FileExists(GroundTruthFilename))
-	{
-		UE_LOG(LogTemp, Log, TEXT("GroundTruth file located at %s"), *GroundTruthFilename);
-	}
-	else
+	TArray<FString> GroundTruthFilenames = ScreenshotManager->FindApprovedFiles(MetaData, TEXT("GroundTruth.json"));
+	if (GroundTruthFilenames.Num() == 0)
 	{
-		UE_LOG(LogTemp, Error, TEXT("Failed to find a GroundTruth file at %s, generating one now. Rerun the test after verifying them!"), *GroundTruthFilename);
-		SaveOutputToGroundTruth(ReportDirectory, InOutputData);
+		FString IdealReportDirectory = ScreenshotManager->GetIdealApprovedFolderForImage(MetaData);
+		UE_LOG(LogTemp, Error, TEXT("Failed to find a GroundTruth file at %s, generating one now. Rerun the test after verifying them!"), *IdealReportDirectory);
+		SaveOutputToGroundTruth(IdealReportDirectory, InOutputData);
 		FinishTest(EFunctionalTestResult::Failed, TEXT("Generated ground truth, run test again after verifying the ground truth is correct."));
 		return;
 	}
+	FString GroundTruthFilename = GroundTruthFilenames[0];
+	UE_LOG(LogTemp, Log, TEXT("GroundTruth file located at %s"), *GroundTruthFilename);
+	FString ReportDirectory = FPaths::GetPath(GroundTruthFilename);
 
-	// The groundtruth file exists, so we can load it and turn it back into a FMoviePipelineOutputData struct.
+	// The ground truth file exists, so we can load it and turn it back into a FMoviePipelineOutputData struct.
 	FString LoadedGroundTruthJsonStr;
 	FFileHelper::LoadFileToString(LoadedGroundTruthJsonStr, *GroundTruthFilename);
 
@@ -190,41 +188,34 @@ void AMoviePipelineFunctionalTestBase::CompareRenderOutputToGroundTruth(FMoviePi
 	int32 NumShotsGT = GroundTruthData.ShotData.Num();
 	if (NumShotsNew != NumShotsGT)
 	{
-		FinishTest(EFunctionalTestResult::Failed, TEXT("Mismatched number of shots between GT and New."));
+		FinishTest(EFunctionalTestResult::Failed, FString::Printf(TEXT("Mismatched number of shots between GroudTruth and New. Expected %d got %d."), NumShotsGT, NumShotsNew));
 		return;
 	}
 
 	TMap<FString, FString> OldToNewImagesToCompare;
 
-	for (const FMoviePipelineShotOutputData& GroundTruthShot : GroundTruthData.ShotData)
+	for (int32 ShotIndex = 0; ShotIndex < GroundTruthData.ShotData.Num(); ShotIndex++)
 	{
-		// Look up the matching shot in the new data.
-		const FMoviePipelineShotOutputData* FoundInNewData = InOutputData.ShotData.FindByPredicate([GroundTruthShot](const FMoviePipelineShotOutputData& InOther)
-			{
-				return GroundTruthShot.Shot->InnerName == InOther.Shot->InnerName &&
-					GroundTruthShot.Shot->OuterName == InOther.Shot->OuterName;
-			});
-		if (!FoundInNewData)
-		{
-			FinishTest(EFunctionalTestResult::Failed, TEXT("Did not render a shot that was in the Ground Truth."));
-			return;
-		}
+		// We compare them using indexes because we know that they have the same number of indexes,
+		// and the order should be deterministic. 
+		const FMoviePipelineShotOutputData & GroundTruthShot = GroundTruthData.ShotData[ShotIndex];
+		const FMoviePipelineShotOutputData & NewData = InOutputData.ShotData[ShotIndex];
 
 		// And then repeat this for each render pass in the shot
-		int32 NumRenderPassesNew = FoundInNewData->RenderPassData.Num();
+		int32 NumRenderPassesNew = NewData.RenderPassData.Num();
 		int32 NumRenderPassesGT = GroundTruthShot.RenderPassData.Num();
 		if (NumRenderPassesNew != NumRenderPassesGT)
 		{
-			FinishTest(EFunctionalTestResult::Failed, TEXT("Mismatched number of Render Passes between GT and New."));
+			FinishTest(EFunctionalTestResult::Failed, FString::Printf(TEXT("Mismatched number of shots between GroudTruth and New. Expected %d got %d."), NumRenderPassesGT, NumRenderPassesNew));
 			return;
 		}
 
 		for (const TPair<FMoviePipelinePassIdentifier, FMoviePipelineRenderPassOutputData>& GroundTruthPair : GroundTruthShot.RenderPassData)
 		{
-			const FMoviePipelineRenderPassOutputData* NewDataOutputData = FoundInNewData->RenderPassData.Find(GroundTruthPair.Key);
+			const FMoviePipelineRenderPassOutputData* NewDataOutputData = NewData.RenderPassData.Find(GroundTruthPair.Key);
 			if (!NewDataOutputData)
 			{
-				FinishTest(EFunctionalTestResult::Failed, TEXT("Did not render a render pass that was in the Ground Truth."));
+				FinishTest(EFunctionalTestResult::Failed, FString::Printf(TEXT("Did not render pass '%s' from camera '%s' that is in the Ground Truth."), *GroundTruthPair.Key.Name, *GroundTruthPair.Key.CameraName));
 				return;
 			}
 
@@ -233,7 +224,7 @@ void AMoviePipelineFunctionalTestBase::CompareRenderOutputToGroundTruth(FMoviePi
 			int32 NumOutputFilesGT = GroundTruthPair.Value.FilePaths.Num();
 			if (NumOutputFilesNew != NumOutputFilesGT)
 			{
-				FinishTest(EFunctionalTestResult::Failed, TEXT("Mismatched number of files on disk between GT and New."));
+				FinishTest(EFunctionalTestResult::Failed, FString::Printf(TEXT("Mismatched number of shots between GroudTruth and New. Expected %d got %d."), NumOutputFilesGT, NumOutputFilesNew));
 				return;
 			}
 
@@ -249,28 +240,16 @@ void AMoviePipelineFunctionalTestBase::CompareRenderOutputToGroundTruth(FMoviePi
 	}
 
 	// Time for the computationally expensive part, doing image comparisons!
-	for (const TPair<FString, FString>& Pair : OldToNewImagesToCompare)
+	TSharedPtr<FImageComparisonResult> ComparisonResult = ScreenshotManager->CompareImageSequence(OldToNewImagesToCompare, MetaData);
+	if (ComparisonResult.IsValid() && !ComparisonResult->AreSimilar())
 	{
-		const FString& OldImagePath = Pair.Key;
-		const FString& NewImagePath = Pair.Value;
-
-		// Calculate a path for the delta image to be saved.
-		FString DeltaPath = FPaths::ChangeExtension(NewImagePath, TEXT(""));
-		FString OldExtension = FPaths::GetExtension(NewImagePath, true);
-		DeltaPath += TEXT("_Delta") + OldExtension;
-
-		// Alright we have both images in memory now, now use a FImageComparer for less strict comparison.
-		FImageTolerance Tolerance = FImageTolerance::DefaultIgnoreLess;
-		FImageComparer Comparer;
-		FImageComparisonResult ComparisonResults = Comparer.Compare(OldImagePath, NewImagePath, Tolerance, DeltaPath);
-
-		if (!ComparisonResults.AreSimilar())
-		{
-			FinishTest(EFunctionalTestResult::Failed, TEXT("Frames failed comparison tolerance!"));
-			return;
-		}
+		ScreenshotManager->NotifyAutomationTestFrameworkOfImageComparison(*ComparisonResult);
+		FinishTest(EFunctionalTestResult::Failed, TEXT("Frames failed comparison tolerance!"));
+		return;
 	}
 
+	UE_LOG(LogTemp, Display, TEXT("All image sequences from %s are similar to the Ground Truth."), *MetaData.ScreenShotName);
+
 	FinishTest(EFunctionalTestResult::Succeeded, TEXT(""));
 }
 
diff --git a/Engine/Source/Developer/AutomationController/Private/AutomationControllerManager.cpp b/Engine/Source/Developer/AutomationController/Private/AutomationControllerManager.cpp
--- a/Engine/Source/Developer/AutomationController/Private/AutomationControllerManager.cpp
+++ b/Engine/Source/Developer/AutomationController/Private/AutomationControllerManager.cpp
@@ -447,6 +447,48 @@ void FAutomationControllerManager::Tick()
 	ProcessComparisonQueue();
 }
 
+void FAutomationControllerManager::ReportImageComparisonResult(const FMessageAddress Sender, const FAutomationWorkerImageComparisonResults& Result)
+{
+	// Find the game session instance info
+	int32 ClusterIndex;
+	int32 DeviceIndex;
+	verify(DeviceClusterManager.FindDevice(Sender, ClusterIndex, DeviceIndex));
+
+	// Get the current test.
+	TSharedPtr<IAutomationReport> Report = DeviceClusterManager.GetTest(ClusterIndex, DeviceIndex);
+	if (Report.IsValid())
+	{
+		// Record the artifacts for the test.
+		TMap<EComparisonFileTypes, FString> LocalFiles;
+
+		FString ScreenshotResultsFolder = FPaths::AutomationReportsDir();
+
+		// Paths in the result are relative to the automation report directory.	
+		LocalFiles.Add(EComparisonFileTypes::Unapproved, FPaths::Combine(ScreenshotResultsFolder, Result.ReportIncomingFilePath));
+
+		// Don't copy reference and delta if the images are similar.
+		if (!Result.bSimilar)
+		{
+			// unapproved should always be valid. but approved/difference may be empty if this is a new screenshot
+			if (Result.ReportIncomingFilePath.Len())
+			{
+				LocalFiles.Add(EComparisonFileTypes::Approved, FPaths::Combine(ScreenshotResultsFolder, Result.ReportApprovedFilePath));
+			}
+
+			if (Result.ReportComparisonFilePath.Len())
+			{
+				LocalFiles.Add(EComparisonFileTypes::Difference, FPaths::Combine(ScreenshotResultsFolder, Result.ReportComparisonFilePath));
+			}
+		}
+
+		Report->AddArtifact(ClusterIndex, CurrentTestPass, FAutomationArtifact(Result.UniqueId, Result.ScreenshotName, EAutomationArtifactType::Comparison, LocalFiles));
+	}
+	else
+	{
+		UE_LOG(LogAutomationController, Error, TEXT("Cannot generate screenshot report for screenshot %s as report is missing"), *Result.IncomingFilePath);
+	}
+}
+
 void FAutomationControllerManager::ProcessComparisonQueue()
 {
 	TSharedPtr<FComparisonEntry> Entry;
@@ -458,66 +500,33 @@ void FAutomationControllerManager::ProcessComparisonQueue()
 			check(Dequeued);
 
 			FImageComparisonResult Result = Entry->PendingComparison.Get();
-
-			const FGuid UniqueId = FGuid::NewGuid();
+			FAutomationWorkerImageComparisonResults ResultMessage(
+				FGuid::NewGuid(),
+				Result.ScreenshotName,
+				Result.IsNew(),
+				Result.AreSimilar(),
+				Result.MaxLocalDifference,
+				Result.GlobalDifference,
+				Result.ErrorMessage.ToString(),
+				Result.IncomingFilePath,
+				Result.ReportComparisonFilePath,
+				Result.ReportApprovedFilePath,
+				Result.ReportIncomingFilePath
+			);
 
 			// Send the message back to the automation worker letting it know the results of the comparison test.
 			{
-				FAutomationWorkerImageComparisonResults* Message = FMessageEndpoint::MakeMessage<FAutomationWorkerImageComparisonResults>(
-					UniqueId,
-					Result.IsNew(),
-					Result.AreSimilar(),
-					Result.MaxLocalDifference,
-					Result.GlobalDifference,
-					Result.ErrorMessage.ToString()
-				);
+				FAutomationWorkerImageComparisonResults* Message = FMessageEndpoint::MakeMessage<FAutomationWorkerImageComparisonResults>(ResultMessage);
 
-				UE_LOG(LogAutomationController, Log, TEXT("Sending ImageComparisonResult to %s (IsNew=%d, AreSimilar=%d)"), 
+				UE_LOG(LogAutomationController, Log, TEXT("Sending ImageComparisonResult to %s (IsNew=%d, AreSimilar=%d)"),
 					*Entry->Sender.ToString()
 					, Result.IsNew()
 					, Result.AreSimilar()
-					);
+				);
 				MessageEndpoint->Send(Message, Entry->Sender);
 			}
 
-			// Find the game session instance info
-			int32 ClusterIndex;
-			int32 DeviceIndex;
-			verify(DeviceClusterManager.FindDevice(Entry->Sender, ClusterIndex, DeviceIndex));
-
-			// Get the current test.
-			TSharedPtr<IAutomationReport> Report = DeviceClusterManager.GetTest(ClusterIndex, DeviceIndex);
-			if (Report.IsValid())
-			{
-				// Record the artifacts for the test.
-				TMap<EComparisonFileTypes, FString> LocalFiles;
-
-				FString ScreenshotResultsFolder = FPaths::AutomationReportsDir();
-	
-				// Paths in the result are relative to the automation report directory.	
-				LocalFiles.Add(EComparisonFileTypes::Unapproved, FPaths::Combine(ScreenshotResultsFolder, Result.ReportIncomingFilePath));
-
-				// Don't copy reference and delta if the images are similar.
-				if (!Result.AreSimilar())
-				{
-					// unapproved should always be valid. but approved/difference may be empty if this is a new screenshot
-					if (Result.ReportIncomingFilePath.Len())
-					{
-						LocalFiles.Add(EComparisonFileTypes::Approved, FPaths::Combine(ScreenshotResultsFolder, Result.ReportApprovedFilePath));
-					}
-
-					if (Result.ReportComparisonFilePath.Len())
-					{
-						LocalFiles.Add(EComparisonFileTypes::Difference, FPaths::Combine(ScreenshotResultsFolder, Result.ReportComparisonFilePath));
-					}
-				}
-
-				Report->AddArtifact(ClusterIndex, CurrentTestPass, FAutomationArtifact(UniqueId, Entry->ScreenshotPath, EAutomationArtifactType::Comparison, LocalFiles));
-			}
-			else
-			{
-				UE_LOG(LogAutomationController, Error, TEXT("Cannot generate screenshot report for screenshot %s as report is missing"), *Result.IncomingFilePath);
-			}
+			ReportImageComparisonResult(Entry->Sender, ResultMessage);
 		}
 	}
 }
@@ -876,6 +885,7 @@ void FAutomationControllerManager::Startup()
 		.Handling<FAutomationWorkerRequestTestsReplyComplete>(this, &FAutomationControllerManager::HandleRequestTestsReplyCompleteMessage)
 		.Handling<FAutomationWorkerRunTestsReply>(this, &FAutomationControllerManager::HandleRunTestsReplyMessage)
 		.Handling<FAutomationWorkerScreenImage>(this, &FAutomationControllerManager::HandleReceivedScreenShot)
+		.Handling<FAutomationWorkerImageComparisonResults>(this, &FAutomationControllerManager::HandleReceivedComparisonResult)
 		.Handling<FAutomationWorkerTestDataRequest>(this, &FAutomationControllerManager::HandleTestDataRequest)
 		.Handling<FAutomationWorkerWorkerOffline>(this, &FAutomationControllerManager::HandleWorkerOfflineMessage)
 		.Handling<FAutomationWorkerTelemetryData>(this, &FAutomationControllerManager::HandleReceivedTelemetryData);
@@ -1314,12 +1324,16 @@ void FAutomationControllerManager::HandleReceivedScreenShot(const FAutomationWor
 	// compare the incoming image and throw it away afterwards (note - there will be a copy in the report)
 	TSharedRef<FComparisonEntry> Comparison = MakeShareable(new FComparisonEntry());
 	Comparison->Sender = Context->GetSender();
-	Comparison->ScreenshotPath = Message.Metadata.Context /  Message.Metadata.ScreenShotName;
 	Comparison->PendingComparison = ScreenshotManager->CompareScreenshotAsync(IncomingFileName, Message.Metadata, EScreenShotCompareOptions::DiscardImage);
 
 	ComparisonQueue.Enqueue(Comparison);
 }
 
+void FAutomationControllerManager::HandleReceivedComparisonResult(const FAutomationWorkerImageComparisonResults& Message, const TSharedRef<IMessageContext, ESPMode::ThreadSafe>& Context)
+{
+	ReportImageComparisonResult(Context->GetSender(), Message);
+}
+
 void FAutomationControllerManager::HandleTestDataRequest(const FAutomationWorkerTestDataRequest& Message, const TSharedRef<IMessageContext, ESPMode::ThreadSafe>& Context)
 {
 	UE_LOG(LogAutomationController, Log, TEXT("Received TestDataRequest from %s"), *Context->GetSender().ToString());
diff --git a/Engine/Source/Developer/AutomationController/Private/AutomationControllerManager.h b/Engine/Source/Developer/AutomationController/Private/AutomationControllerManager.h
--- a/Engine/Source/Developer/AutomationController/Private/AutomationControllerManager.h
+++ b/Engine/Source/Developer/AutomationController/Private/AutomationControllerManager.h
@@ -457,6 +457,9 @@ protected:
 	 */
 	void ExecuteNextTask( int32 ClusterIndex, OUT bool& bAllTestsCompleted );
 
+	/* Report an image comparison result */
+	void ReportImageComparisonResult(const FMessageAddress Sender, const FAutomationWorkerImageComparisonResults& Result);
+
 	/** Process the comparison queue to see if there are comparisons we need to respond to the test with. */
 	void ProcessComparisonQueue();
 
@@ -493,6 +496,9 @@ private:
 	/** Handles FAutomationWorkerScreenImage messages. */
 	void HandleReceivedScreenShot( const FAutomationWorkerScreenImage& Message, const TSharedRef<IMessageContext, ESPMode::ThreadSafe>& Context );
 
+	/** Handles FAutomationWorkerScreenshotComparisonResult messages. */
+	void HandleReceivedComparisonResult( const FAutomationWorkerImageComparisonResults& Message, const TSharedRef<IMessageContext, ESPMode::ThreadSafe>& Context );
+
 	/** Handles FAutomationWorkerTestDataRequest messages. */
 	void HandleTestDataRequest(const FAutomationWorkerTestDataRequest& Message, const TSharedRef<IMessageContext, ESPMode::ThreadSafe>& Context);
 
@@ -622,7 +628,6 @@ private:
 	struct FComparisonEntry
 	{
 		FMessageAddress Sender;
-		FString ScreenshotPath;
 		TFuture<FImageComparisonResult> PendingComparison;
 	};
 
diff --git a/Engine/Source/Developer/FunctionalTesting/Private/AutomationBlueprintFunctionLibrary.cpp b/Engine/Source/Developer/FunctionalTesting/Private/AutomationBlueprintFunctionLibrary.cpp
--- a/Engine/Source/Developer/FunctionalTesting/Private/AutomationBlueprintFunctionLibrary.cpp
+++ b/Engine/Source/Developer/FunctionalTesting/Private/AutomationBlueprintFunctionLibrary.cpp
@@ -493,7 +493,7 @@ public:
 
 		if (FAutomationTestBase* CurrentTest = FAutomationTestFramework::Get().GetCurrentTest())
 		{
-			CurrentTest->AddEvent(CompareResults.ToAutomationEvent(ScreenShotName));
+			CurrentTest->AddEvent(CompareResults.ToAutomationEvent());
 		}
 
 		DeleteSelfNextFrame();
@@ -586,7 +586,7 @@ public:
 
 		if (FAutomationTestBase* CurrentTest = FAutomationTestFramework::Get().GetCurrentTest())
 		{
-			CurrentTest->AddEvent(CompareResults.ToAutomationEvent(ScreenShotName));
+			CurrentTest->AddEvent(CompareResults.ToAutomationEvent());
 		}
 
 		delete this;
diff --git a/Engine/Source/Developer/FunctionalTesting/Private/ScreenshotFunctionalTestBase.cpp b/Engine/Source/Developer/FunctionalTesting/Private/ScreenshotFunctionalTestBase.cpp
--- a/Engine/Source/Developer/FunctionalTesting/Private/ScreenshotFunctionalTestBase.cpp
+++ b/Engine/Source/Developer/FunctionalTesting/Private/ScreenshotFunctionalTestBase.cpp
@@ -207,7 +207,7 @@ void AScreenshotFunctionalTestBase::OnComparisonComplete(const FAutomationScreen
 
 	if (FAutomationTestBase* CurrentTest = FAutomationTestFramework::Get().GetCurrentTest())
 	{
-		CurrentTest->AddEvent(CompareResults.ToAutomationEvent(TestLabel));
+		CurrentTest->AddEvent(CompareResults.ToAutomationEvent());
 	}
 
 	FAutomationTestFramework::Get().NotifyScreenshotTakenAndCompared();
diff --git a/Engine/Source/Developer/ScreenShotComparisonTools/Private/ScreenShotManager.cpp b/Engine/Source/Developer/ScreenShotComparisonTools/Private/ScreenShotManager.cpp
--- a/Engine/Source/Developer/ScreenShotComparisonTools/Private/ScreenShotManager.cpp
+++ b/Engine/Source/Developer/ScreenShotComparisonTools/Private/ScreenShotManager.cpp
@@ -22,6 +22,8 @@
 
 DEFINE_LOG_CATEGORY(LogScreenShotManager);
 
+#define LOCTEXT_NAMESPACE "ScreenShotManager"
+
 class FScreenshotComparisons
 {
 public:
@@ -85,8 +87,6 @@ FString FScreenShotManager::GetPathComponentForTestImages(const FAutomationScree
 
 FString FScreenShotManager::GetApprovedFolderForImageWithOptions(const FAutomationScreenshotMetadata& MetaData, EApprovedFolderOptions InOptions) const
 {
-	// plaform will be something like PS4
-	// RHI = SM5
 	const FDataDrivenPlatformInfo& PlatInfo = FDataDrivenPlatformInfoRegistry::GetPlatformInfo(MetaData.Platform);
 
 	bool bUsePlatformPath = PlatInfo.bIsConfidential && (InOptions & EApprovedFolderOptions::UsePlatformFolders) == 0;
@@ -129,13 +129,13 @@ FString FScreenShotManager::GetIdealApprovedFolderForImage(const FAutomationScre
 	return GetApprovedFolderForImageWithOptions(MetaData, DefaultOptions);
 }
 
-TArray<FString> FScreenShotManager::FindApprovedImages(const FAutomationScreenshotMetadata& IncomingMetaData)
+TArray<FString> FScreenShotManager::FindApprovedFiles(const FAutomationScreenshotMetadata& IncomingMetaData, const FString& FilePattern) const
 {
 	TArray<FString> TriedPaths;
 
-	auto FindImages = [&TriedPaths](TArray<FString>& OutApprovedImages, const FString& InPath)
+	auto FindImages = [&TriedPaths, FilePattern](TArray<FString>& OutApprovedImages, const FString& InPath)
 	{
-		IFileManager::Get().FindFilesRecursive(OutApprovedImages, *InPath, TEXT("*.png"), true, false);
+		IFileManager::Get().FindFilesRecursive(OutApprovedImages, *InPath, *FilePattern, true, false);
 		TriedPaths.Emplace(InPath);
 	};
 
@@ -172,7 +172,7 @@ TArray<FString> FScreenShotManager::FindApprovedImages(const FAutomationScreensh
 		}
 	}
 
-	// @todo(agrant): find fallback images if they don't exist at this point
+	// find fallback images if they don't exist at this point
 	if (ApprovedImages.Num() == 0)
 	{
 		FString CurrentPlatformRHI = GetPathComponentForPlatformAndRHI(IncomingMetaData);
@@ -181,7 +181,7 @@ TArray<FString> FScreenShotManager::FindApprovedImages(const FAutomationScreensh
 
 		while (ApprovedImages.Num() == 0)
 		{
-			FString* FallbackPlatformRHI = FallbackPlatforms.Find(CurrentPlatformRHI);
+			const FString* FallbackPlatformRHI = FallbackPlatforms.Find(CurrentPlatformRHI);
 			if (!FallbackPlatformRHI)
 			{
 				break;
@@ -287,7 +287,7 @@ FImageComparisonResult FScreenShotManager::CompareScreenshot(const FString& InUn
 	DefaultTolerance.IgnoreAntiAliasing = true;
 
 	// find all the approved images we can use. This will find fallback images from other platforms if necessary
-	TArray<FString> ApprovedDeviceShots = FindApprovedImages(IncomingMetaData);
+	TArray<FString> ApprovedDeviceShots = FindApprovedFiles(IncomingMetaData, TEXT("*.png"));
 
 	TOptional<FAutomationScreenshotMetadata> NearestExistingApprovedImageMetadata;
 
@@ -381,6 +381,8 @@ FImageComparisonResult FScreenShotManager::CompareScreenshot(const FString& InUn
 	ComparisonResult.SourcePlatform = IncomingMetaData.Platform;
 	ComparisonResult.SourceRHI = IncomingMetaData.Rhi;
 	ComparisonResult.IdealApprovedFolderPath = IdealApprovedFolderPath;
+	// We use the subfolder path to the screenshot as name (before any environment specialization - platform, RHI - are appended).
+	ComparisonResult.ScreenshotName = ResultsSubFolder;
 
 	// Result paths should be relative to the project. Note this may be empty, and if it is MakePathRelative returns
 	// a non empty relative path... but we want it to stay empty as that's how we signal that no approved file exists
@@ -590,6 +592,114 @@ bool FScreenShotManager::OpenComparisonReports(FString ImportPath, TArray<FCompa
 	return true;
 }
 
+TSharedPtr<FImageComparisonResult> FScreenShotManager::CompareImageSequence(const TMap<FString, FString>& Sequence, const FAutomationScreenshotMetadata& Metadata)
+{
+	// Time for the computationally expensive part, doing image comparisons!
+	for (const TPair<FString, FString>& Pair : Sequence)
+	{
+		const FString& OldImagePath = Pair.Key;
+		const FString& NewImagePath = Pair.Value;
+
+		// Calculate a path for the delta image to be saved.
+		FString DeltaPath = FPaths::ChangeExtension(NewImagePath, TEXT(""));
+		FString OldExtension = FPaths::GetExtension(NewImagePath, true);
+		DeltaPath += TEXT("_Delta") + OldExtension;
+
+		// Alright we have both images in memory now, now use a FImageComparer for less strict comparison.
+		FImageTolerance Tolerance = FImageTolerance::DefaultIgnoreLess;
+		FImageComparer Comparer;
+		FImageComparisonResult ComparisonResult = Comparer.Compare(OldImagePath, NewImagePath, Tolerance, DeltaPath);
+
+		if (!ComparisonResult.AreSimilar())
+		{
+			ComparisonResult.SourcePlatform = Metadata.Platform;
+			ComparisonResult.SourceRHI = Metadata.Rhi;
+			ComparisonResult.IdealApprovedFolderPath = GetIdealApprovedFolderForImage(Metadata);;
+			ComparisonResult.ScreenshotName = GetPathComponentForTestImages(Metadata);
+			FString ReportPathOnDisk = FPaths::Combine(ScreenshotResultsFolder, GetPathComponentForTestImages(Metadata), Metadata.Platform, GetPathComponentForRHI(Metadata), TEXT("/"));
+
+			ComparisonResult.ReportApprovedFilePath = FPaths::Combine(ReportPathOnDisk, TEXT("Approved.png"));
+			ComparisonResult.ReportIncomingFilePath = FPaths::Combine(ReportPathOnDisk, TEXT("Incoming.png"));
+			ComparisonResult.ReportComparisonFilePath = FPaths::Combine(ReportPathOnDisk, TEXT("Delta.png"));
+
+			TArray<TPair<FString, FString>> FilesToCopy;
+			FilesToCopy.Emplace(ComparisonResult.ApprovedFilePath, FPaths::ConvertRelativePathToFull(ComparisonResult.ReportApprovedFilePath));
+			FilesToCopy.Emplace(ComparisonResult.IncomingFilePath, FPaths::ConvertRelativePathToFull(ComparisonResult.ReportIncomingFilePath));
+			FilesToCopy.Emplace(ComparisonResult.ComparisonFilePath, FPaths::ConvertRelativePathToFull(ComparisonResult.ReportComparisonFilePath));
+
+			ParallelFor(FilesToCopy.Num(), [&](int32 Index)
+				{
+					const TPair<FString, FString>& FilePathPair = FilesToCopy[Index];
+					IFileManager::Get().Copy(*FilePathPair.Value, *FilePathPair.Key, true, true);
+				});
+
+			FPaths::MakePathRelativeTo(ComparisonResult.ReportApprovedFilePath, *ScreenshotResultsFolder);
+			FPaths::MakePathRelativeTo(ComparisonResult.ReportIncomingFilePath, *ScreenshotResultsFolder);
+			FPaths::MakePathRelativeTo(ComparisonResult.ReportComparisonFilePath, *ScreenshotResultsFolder);
+			FPaths::MakePathRelativeTo(ComparisonResult.IncomingFilePath, *FPaths::ProjectDir());
+			FPaths::MakePathRelativeTo(ComparisonResult.IdealApprovedFolderPath, *FPaths::ProjectDir());
+
+			if (!ComparisonResult.ApprovedFilePath.IsEmpty())
+			{
+				FPaths::MakePathRelativeTo(ComparisonResult.ApprovedFilePath, *FPaths::ProjectDir());
+			}
+			if (!ComparisonResult.ComparisonFilePath.IsEmpty())
+			{
+				FPaths::MakePathRelativeTo(ComparisonResult.ComparisonFilePath, *FPaths::ProjectDir());
+			}
+
+			if (ComparisonResult.ErrorMessage.IsEmpty())
+			{
+				ComparisonResult.ErrorMessage = FText::Format(
+					LOCTEXT("CompareImageSequenceError", "Frame '{0}' is not similar to Ground Truth"),
+					FText::FromString(FPaths::GetBaseFilename(ComparisonResult.ApprovedFilePath))
+				);
+			}
+
+			// save the result at to a report
+			FString Json;
+			if (FJsonObjectConverter::UStructToJsonObjectString(ComparisonResult, Json))
+			{
+				FString ComparisonReportFile = FPaths::Combine(ReportPathOnDisk, TEXT("Report.json"));
+				FFileHelper::SaveStringToFile(Json, *ComparisonReportFile, FFileHelper::EEncodingOptions::ForceUTF8WithoutBOM);
+
+				UE_LOG(LogScreenShotManager, Log, TEXT("Saved report to %s"), *ComparisonReportFile);
+			}
+			if (FJsonObjectConverter::UStructToJsonObjectString(Metadata, Json))
+			{
+				FString MetadataPath = FPaths::Combine(ReportPathOnDisk, TEXT("Incoming.json"));
+				FFileHelper::SaveStringToFile(Json, *MetadataPath, FFileHelper::EEncodingOptions::ForceUTF8WithoutBOM);
+			}
+
+			return MakeShareable(new FImageComparisonResult(ComparisonResult));
+		}
+	}
+
+	return nullptr;
+}
+
+void FScreenShotManager::NotifyAutomationTestFrameworkOfImageComparison(const FImageComparisonResult& ComparisonResult)
+{
+	FAutomationScreenshotCompareResults Results(
+		FGuid::NewGuid(),
+		ComparisonResult.ErrorMessage.ToString(),
+		ComparisonResult.MaxLocalDifference,
+		ComparisonResult.GlobalDifference,
+		ComparisonResult.IsNew(),
+		ComparisonResult.AreSimilar(),
+		ComparisonResult.IncomingFilePath,
+		ComparisonResult.ReportComparisonFilePath,
+		ComparisonResult.ReportApprovedFilePath,
+		ComparisonResult.ReportIncomingFilePath,
+		ComparisonResult.ScreenshotName
+	);
+	FAutomationTestFramework::Get().NotifyScreenshotComparisonReport(Results);
+	if (FAutomationTestBase* CurrentTest = FAutomationTestFramework::Get().GetCurrentTest())
+	{
+		CurrentTest->AddEvent(Results.ToAutomationEvent());
+	}
+}
+
 FString FScreenShotManager::GetDefaultExportDirectory() const
 {
 	return FPaths::Combine(FPaths::ProjectSavedDir(),TEXT("Exported/imageCompare"));
diff --git a/Engine/Source/Developer/ScreenShotComparisonTools/Private/ScreenShotManager.h b/Engine/Source/Developer/ScreenShotComparisonTools/Private/ScreenShotManager.h
--- a/Engine/Source/Developer/ScreenShotComparisonTools/Private/ScreenShotManager.h
+++ b/Engine/Source/Developer/ScreenShotComparisonTools/Private/ScreenShotManager.h
@@ -52,6 +52,13 @@ public:
 	virtual bool OpenComparisonReports(FString ImportPath, TArray<FComparisonReport>& OutReports) override;
 	
 	virtual FString GetIdealApprovedFolderForImage(const FAutomationScreenshotMetadata& MetaData) const override;
+
+	virtual TArray<FString> FindApprovedFiles(const FAutomationScreenshotMetadata& IncomingMetaData, const FString& FilePattern) const override;
+
+	virtual TSharedPtr<FImageComparisonResult> CompareImageSequence(const TMap<FString, FString>& Sequence, const FAutomationScreenshotMetadata& Metadata) override;
+
+	virtual void NotifyAutomationTestFrameworkOfImageComparison(const FImageComparisonResult& Result) override;
+
 	//~ End IScreenShotManager Interface
 
 private:
@@ -63,8 +70,6 @@ private:
 
 	FString GetApprovedFolderForImageWithOptions(const FAutomationScreenshotMetadata& MetaData, EApprovedFolderOptions InOptions) const;
 
-	TArray<FString> FindApprovedImages(const FAutomationScreenshotMetadata& IncomingMetaData);
-
 	FString GetDefaultExportDirectory() const;
 
 	FImageComparisonResult CompareScreenshot(const FString& IncomingPath, const FAutomationScreenshotMetadata& MetaData, const EScreenShotCompareOptions Options);
diff --git a/Engine/Source/Developer/ScreenShotComparisonTools/Public/ImageComparer.h b/Engine/Source/Developer/ScreenShotComparisonTools/Public/ImageComparer.h
--- a/Engine/Source/Developer/ScreenShotComparisonTools/Public/ImageComparer.h
+++ b/Engine/Source/Developer/ScreenShotComparisonTools/Public/ImageComparer.h
@@ -302,6 +302,12 @@ public:
 	UPROPERTY()
 	FText ErrorMessage;
 
+	/*
+		Name of the screenshot
+	*/
+	UPROPERTY()
+	FString ScreenshotName;
+
 	/*
 		Version of the image comparision result 
 	*/
diff --git a/Engine/Source/Developer/ScreenShotComparisonTools/Public/Interfaces/IScreenShotManager.h b/Engine/Source/Developer/ScreenShotComparisonTools/Public/Interfaces/IScreenShotManager.h
--- a/Engine/Source/Developer/ScreenShotComparisonTools/Public/Interfaces/IScreenShotManager.h
+++ b/Engine/Source/Developer/ScreenShotComparisonTools/Public/Interfaces/IScreenShotManager.h
@@ -80,4 +80,19 @@ public:
 	* based on the platform/rhi/etc.
 	*/
 	virtual FString GetIdealApprovedFolderForImage(const FAutomationScreenshotMetadata& MetaData) const = 0;
+
+	/**
+	* Find the all the files that are approved for the given metadata and file pattern 
+	*/
+	virtual TArray<FString> FindApprovedFiles(const FAutomationScreenshotMetadata& IncomingMetaData, const FString& FilePattern) const = 0;
+
+	/**
+	* Compare a sequence of images. If the returned value is null no comparison failed. Otherwise return the first failing frame.
+	*/
+	virtual TSharedPtr<FImageComparisonResult> CompareImageSequence(const TMap<FString, FString>& Sequence, const FAutomationScreenshotMetadata& Metadata) = 0;
+
+	/**
+	* Notify the automation test framework of an image comparison result
+	*/
+	virtual void NotifyAutomationTestFrameworkOfImageComparison(const FImageComparisonResult& Result) = 0;
 };
diff --git a/Engine/Source/Runtime/AutomationMessages/Public/AutomationWorkerMessages.h b/Engine/Source/Runtime/AutomationMessages/Public/AutomationWorkerMessages.h
--- a/Engine/Source/Runtime/AutomationMessages/Public/AutomationWorkerMessages.h
+++ b/Engine/Source/Runtime/AutomationMessages/Public/AutomationWorkerMessages.h
@@ -2,11 +2,11 @@
 
 #pragma once
 
+#include "AutomationState.h"
 #include "CoreTypes.h"
 #include "Containers/UnrealString.h"
-#include "Misc/Guid.h"
 #include "Misc/AutomationTest.h"
-#include "AutomationState.h"
+#include "Misc/Guid.h"
 #include "UObject/ObjectMacros.h"
 
 #include "AutomationWorkerMessages.generated.h"
@@ -692,7 +692,6 @@ public:
 	}
 };
 
-
 /**
  * Implements a message that is sent in containing a screen shot run during performance test.
  */
@@ -737,13 +736,30 @@ public:
 	{
 	}
 
-	FAutomationWorkerImageComparisonResults(FGuid InUniqueId, bool InIsNew, bool InAreSimilar, double InMaxLocalDifference, double InGlobalDifference, FString InErrorMessage)
+	FAutomationWorkerImageComparisonResults(
+		FGuid InUniqueId,
+		const FString& InName,
+		bool InIsNew,
+		bool InAreSimilar,
+		double InMaxLocalDifference,
+		double InGlobalDifference,
+		const FString& InErrorMessage,
+		const FString& InIncomingFilePath,
+		const FString& InReportComparisonFilePath,
+		const FString& InReportApprovedFilePath,
+		const FString& InReportIncomingFilePath
+	)
 		: UniqueId(InUniqueId)
+		, ScreenshotName(InName)
 		, bNew(InIsNew)
 		, bSimilar(InAreSimilar)
 		, MaxLocalDifference(InMaxLocalDifference)
 		, GlobalDifference(InGlobalDifference)
 		, ErrorMessage(InErrorMessage)
+		, IncomingFilePath(InIncomingFilePath)
+		, ReportComparisonFilePath(InReportComparisonFilePath)
+		, ReportApprovedFilePath(InReportApprovedFilePath)
+		, ReportIncomingFilePath(InReportIncomingFilePath)
 	{
 	}
 
@@ -751,6 +767,10 @@ public:
 	UPROPERTY(EditAnywhere, Category="Message")
 	FGuid UniqueId;
 
+	/** The name of the screenshot. */
+	UPROPERTY(EditAnywhere, Category = "Message")
+	FString ScreenshotName;
+
 	/** Was this a new image we've never seen before and have no ground truth for? */
 	UPROPERTY(EditAnywhere, Category="Message")
 	bool bNew;
@@ -767,6 +787,18 @@ public:
 
 	UPROPERTY(EditAnywhere, Category="Message")
 	FString ErrorMessage;
+
+	UPROPERTY(EditAnywhere, Category = "Message")
+	FString IncomingFilePath;
+
+	UPROPERTY(EditAnywhere, Category = "Message")
+	FString ReportComparisonFilePath;
+
+	UPROPERTY(EditAnywhere, Category = "Message")
+	FString ReportApprovedFilePath;
+
+	UPROPERTY(EditAnywhere, Category = "Message")
+	FString ReportIncomingFilePath;
 };
 
 
diff --git a/Engine/Source/Runtime/AutomationWorker/Private/AutomationWorkerModule.cpp b/Engine/Source/Runtime/AutomationWorker/Private/AutomationWorkerModule.cpp
--- a/Engine/Source/Runtime/AutomationWorker/Private/AutomationWorkerModule.cpp
+++ b/Engine/Source/Runtime/AutomationWorker/Private/AutomationWorkerModule.cpp
@@ -374,6 +374,7 @@ void FAutomationWorkerModule::HandlePreTestingEvent()
 #if WITH_ENGINE
 	FAutomationTestFramework::Get().OnScreenshotCaptured().BindRaw(this, &FAutomationWorkerModule::HandleScreenShotCapturedWithName);
 	FAutomationTestFramework::Get().OnScreenshotAndTraceCaptured().BindRaw(this, &FAutomationWorkerModule::HandleScreenShotAndTraceCapturedWithName);
+	FAutomationTestFramework::Get().OnScreenshotComparisonReport.AddRaw(this, &FAutomationWorkerModule::HandleScreenShotComparisonReport);
 #endif
 }
 
@@ -383,6 +384,7 @@ void FAutomationWorkerModule::HandlePostTestingEvent()
 #if WITH_ENGINE
 	FAutomationTestFramework::Get().OnScreenshotAndTraceCaptured().Unbind();
 	FAutomationTestFramework::Get().OnScreenshotCaptured().Unbind();
+	FAutomationTestFramework::Get().OnScreenshotComparisonReport.RemoveAll(this);
 #endif
 }
 
@@ -394,6 +396,7 @@ void FAutomationWorkerModule::HandleScreenShotCompared(const FAutomationWorkerIm
 	// Image comparison finished.
 	FAutomationScreenshotCompareResults CompareResults;
 	CompareResults.UniqueId = Message.UniqueId;
+	CompareResults.ScreenshotName = Message.ScreenshotName;
 	CompareResults.bWasNew = Message.bNew;
 	CompareResults.bWasSimilar = Message.bSimilar;
 	CompareResults.MaxLocalDifference = Message.MaxLocalDifference;
@@ -418,6 +421,25 @@ void FAutomationWorkerModule::HandlePerformanceDataRetrieved(const FAutomationWo
 }
 
 #if WITH_ENGINE
+void FAutomationWorkerModule::HandleScreenShotComparisonReport(const FAutomationScreenshotCompareResults& Results)
+{
+	FAutomationWorkerImageComparisonResults* Message = FMessageEndpoint::MakeMessage<FAutomationWorkerImageComparisonResults>();
+
+	Message->ScreenshotName = Results.ScreenshotName;
+	Message->UniqueId = Results.UniqueId;
+	Message->bNew = Results.bWasNew;
+	Message->bSimilar = Results.bWasSimilar;
+	Message->ErrorMessage = Results.ErrorMessage;
+	Message->MaxLocalDifference = Results.MaxLocalDifference;
+	Message->GlobalDifference = Results.GlobalDifference;
+	Message->IncomingFilePath = Results.IncomingFilePath;
+	Message->ReportComparisonFilePath = Results.ReportComparisonFilePath;
+	Message->ReportApprovedFilePath = Results.ReportApprovedFilePath;
+	Message->ReportIncomingFilePath = Results.ReportIncomingFilePath;
+
+	MessageEndpoint->Send(Message, TestRequesterAddress);
+}
+
 void FAutomationWorkerModule::HandleScreenShotCapturedWithName(const TArray<FColor>& RawImageData, const FAutomationScreenshotData& Data)
 {
 	HandleScreenShotAndTraceCapturedWithName(RawImageData, TArray<uint8>(), Data);
@@ -575,7 +597,7 @@ void FAutomationWorkerModule::HandleRunTestsMessage( const FAutomationWorkerRunT
 			return;
 		}
 
-		FString LogMessage = FString::Format(TEXT("Worker is already running test '%s' from %s. '%s' won't be run."), 
+		FString LogMessage = FString::Format(TEXT("Worker is already running test '{0}' from {1}. '{2}' won't be run."), 
 			{ *BeautifiedTestName, *TestRequesterAddress.ToString(), *Message.BeautifiedTestName });
 		UE_LOG(LogAutomationWorker, Warning, TEXT("%s"), *LogMessage);
 
diff --git a/Engine/Source/Runtime/AutomationWorker/Private/AutomationWorkerModule.h b/Engine/Source/Runtime/AutomationWorker/Private/AutomationWorkerModule.h
--- a/Engine/Source/Runtime/AutomationWorker/Private/AutomationWorkerModule.h
+++ b/Engine/Source/Runtime/AutomationWorker/Private/AutomationWorkerModule.h
@@ -125,6 +125,9 @@ private:
 
 	/** Invoked when we have screen shot and frame trace to send. */
 	void HandleScreenShotAndTraceCapturedWithName(const TArray<FColor>& RawImageData, const TArray<uint8>& CapturedFrameTrace, const FAutomationScreenshotData& Data);
+
+	/** Invoked when we have a screenshot comparison result to send. */
+	void HandleScreenShotComparisonReport(const FAutomationScreenshotCompareResults& Results);
 #endif
 
 	/** Dispatches analytics events to the data collector. */
diff --git a/Engine/Source/Runtime/Core/Private/Misc/AutomationTest.cpp b/Engine/Source/Runtime/Core/Private/Misc/AutomationTest.cpp
--- a/Engine/Source/Runtime/Core/Private/Misc/AutomationTest.cpp
+++ b/Engine/Source/Runtime/Core/Private/Misc/AutomationTest.cpp
@@ -1024,6 +1024,11 @@ void FAutomationTestFramework::NotifyScreenshotComparisonComplete(const FAutomat
 	OnScreenshotCompared.Broadcast(CompareResults);
 }
 
+void FAutomationTestFramework::NotifyScreenshotComparisonReport(const FAutomationScreenshotCompareResults& CompareResults)
+{
+	OnScreenshotComparisonReport.Broadcast(CompareResults);
+}
+
 void FAutomationTestFramework::NotifyTestDataRetrieved(bool bWasNew, const FString& JsonData)
 {
 	OnTestDataRetrieved.Broadcast(bWasNew, JsonData);
@@ -1197,14 +1202,14 @@ void FAutomationTestExecutionInfo::AddError(const FString& ErrorMessage)
 
 //------------------------------------------------------------------------------
 
-FAutomationEvent FAutomationScreenshotCompareResults::ToAutomationEvent(const FString& ScreenhotName) const
+FAutomationEvent FAutomationScreenshotCompareResults::ToAutomationEvent() const
 {
 	FAutomationEvent Event(EAutomationEventType::Info, TEXT(""));
 
 	if (bWasNew)
 	{
 		Event.Type = EAutomationEventType::Warning;
-		Event.Message = FString::Printf(TEXT("New Screenshot '%s' was discovered!  Please add a ground truth version of it."), *ScreenhotName);
+		Event.Message = FString::Printf(TEXT("New Screenshot '%s' was discovered!  Please add a ground truth version of it."), *ScreenshotName);
 	}
 	else
 	{
@@ -1212,7 +1217,7 @@ FAutomationEvent FAutomationScreenshotCompareResults::ToAutomationEvent(const FS
 		{
 			Event.Type = EAutomationEventType::Info;
 			Event.Message = FString::Printf(TEXT("Screenshot '%s' was similar!  Global Difference = %f, Max Local Difference = %f"),
-				*ScreenhotName, GlobalDifference, MaxLocalDifference);
+				*ScreenshotName, GlobalDifference, MaxLocalDifference);
 		}
 		else
 		{
@@ -1221,11 +1226,11 @@ FAutomationEvent FAutomationScreenshotCompareResults::ToAutomationEvent(const FS
 			if (ErrorMessage.IsEmpty())
 			{
 				Event.Message = FString::Printf(TEXT("Screenshot '%s' test failed, Screenshots were different!  Global Difference = %f, Max Local Difference = %f"),
-					*ScreenhotName, GlobalDifference, MaxLocalDifference);
+					*ScreenshotName, GlobalDifference, MaxLocalDifference);
 			}
 			else
 			{
-				Event.Message = FString::Printf(TEXT("Screenshot '%s' test failed; Error = %s"), *ScreenhotName, *ErrorMessage);
+				Event.Message = FString::Printf(TEXT("Screenshot '%s' test failed; Error = %s"), *ScreenshotName, *ErrorMessage);
 			}
 		}
 	}
diff --git a/Engine/Source/Runtime/Core/Public/Misc/AutomationTest.h b/Engine/Source/Runtime/Core/Public/Misc/AutomationTest.h
--- a/Engine/Source/Runtime/Core/Public/Misc/AutomationTest.h
+++ b/Engine/Source/Runtime/Core/Public/Misc/AutomationTest.h
@@ -790,8 +790,47 @@ struct FAutomationScreenshotCompareResults
 	double GlobalDifference = 0.0;
 	bool bWasNew = false;
 	bool bWasSimilar = false;
+	FString IncomingFilePath;
+	FString ReportComparisonFilePath;
+	FString ReportApprovedFilePath;
+	FString ReportIncomingFilePath;
+	FString ScreenshotName;
+
+	FAutomationScreenshotCompareResults()
+		: UniqueId()
+		, MaxLocalDifference(0.0)
+		, GlobalDifference(0.0)
+		, bWasNew(false)
+		, bWasSimilar(false)
+	{ }
 
-	CORE_API FAutomationEvent ToAutomationEvent(const FString& ScreenhotName) const;
+	FAutomationScreenshotCompareResults(
+		FGuid InUniqueId,
+		FString InErrorMessage,
+		double InMaxLocalDifference,
+		double InGlobalDifference,
+		bool InWasNew,
+		bool InWasSimilar,
+		FString InIncomingFilePath,
+		FString InReportComparisonFilePath,
+		FString InReportApprovedFilePath,
+		FString InReportIncomingFilePath,
+		FString InScreenshotName
+	)
+		: UniqueId(InUniqueId)
+		, ErrorMessage(InErrorMessage)
+		, MaxLocalDifference(InMaxLocalDifference)
+		, GlobalDifference(InGlobalDifference)
+		, bWasNew(InWasNew)
+		, bWasSimilar(InWasSimilar)
+		, IncomingFilePath(InIncomingFilePath)
+		, ReportComparisonFilePath(InReportComparisonFilePath)
+		, ReportApprovedFilePath(InReportApprovedFilePath)
+		, ReportIncomingFilePath(InReportIncomingFilePath)
+		, ScreenshotName(InScreenshotName)
+	{ }
+
+	CORE_API FAutomationEvent ToAutomationEvent() const;
 };
 
 enum class EAutomationComparisonToleranceLevel : uint8
@@ -861,6 +900,8 @@ DECLARE_DELEGATE_ThreeParams(FOnTestScreenshotAndTraceCaptured, const TArray<FCo
 
 DECLARE_MULTICAST_DELEGATE_OneParam(FOnTestScreenshotComparisonComplete, const FAutomationScreenshotCompareResults& /*CompareResults*/);
 
+DECLARE_MULTICAST_DELEGATE_OneParam(FOnTestScreenshotComparisonReport, const FAutomationScreenshotCompareResults& /*CompareResults*/);
+
 DECLARE_MULTICAST_DELEGATE_TwoParams(FOnTestDataRetrieved, bool /*bWasNew*/, const FString& /*JsonData*/);
 
 DECLARE_MULTICAST_DELEGATE_TwoParams(FOnPerformanceDataRetrieved, bool /*bSuccess*/, const FString& /*ErrorMessage*/);
@@ -886,6 +927,9 @@ public:
 	/** Called when a screenshot comparison completes. */
 	FOnTestScreenshotComparisonComplete OnScreenshotCompared;
 
+	/** Called when a screenshot comparison result is reported */
+	FOnTestScreenshotComparisonReport OnScreenshotComparisonReport;
+
 	/** Called when the test data is retrieved. */
 	FOnTestDataRetrieved OnTestDataRetrieved;
 
@@ -1102,7 +1146,16 @@ public:
 	 */
 	static CORE_API bool NeedLogBPTestMetadata();
 
+	/**
+	 * Notify that the screenshot comparison has completed
+	 */
 	CORE_API void NotifyScreenshotComparisonComplete(const FAutomationScreenshotCompareResults& CompareResults);
+
+	/**
+	 * Notify the screenshot comparison report to the framework
+	 */
+	CORE_API void NotifyScreenshotComparisonReport(const FAutomationScreenshotCompareResults& CompareResults);
+
 	CORE_API void NotifyTestDataRetrieved(bool bWasNew, const FString& JsonData);
 	CORE_API void NotifyPerformanceDataRetrieved(bool bSuccess, const FString& ErrorMessage);
 
diff --git a/Engine/Source/Runtime/Engine/Private/Tests/AutomationCommon.cpp b/Engine/Source/Runtime/Engine/Private/Tests/AutomationCommon.cpp
--- a/Engine/Source/Runtime/Engine/Private/Tests/AutomationCommon.cpp
+++ b/Engine/Source/Runtime/Engine/Private/Tests/AutomationCommon.cpp
@@ -238,7 +238,7 @@ namespace AutomationCommon
 	{
 	public:
 		FAutomationImageComparisonRequest(const FString& InImageName, const FString& InContext, int32 InWidth, int32 InHeight, const TArray<FColor>& InImageData, const FAutomationComparisonToleranceAmount& InTolerance, const FString& InNotes)
-			: ImageName(InImageName), ImageData(InImageData), Initiate(false), TaskCompleted(false)
+			: ImageData(InImageData), Initiate(false), TaskCompleted(false)
 		{
 			FString Context = InContext;
 			if (Context.IsEmpty())
@@ -253,7 +253,7 @@ namespace AutomationCommon
 				}
 			}
 
-			ComparisonParameters = BuildScreenshotData(Context, TEXT(""), ImageName, InWidth, InHeight);
+			ComparisonParameters = BuildScreenshotData(Context, TEXT(""), InImageName, InWidth, InHeight);
 
 			// Copy the relevant data into the metadata for the screenshot.
 			ComparisonParameters.bHasComparisonRules = true;
@@ -283,7 +283,7 @@ namespace AutomationCommon
 
 			if (FAutomationTestBase* CurrentTest = FAutomationTestFramework::Get().GetCurrentTest())
 			{
-				CurrentTest->AddEvent(CompareResults.ToAutomationEvent(ImageName));
+				CurrentTest->AddEvent(CompareResults.ToAutomationEvent());
 			}
 
 			TaskCompleted = true;
@@ -309,7 +309,6 @@ namespace AutomationCommon
 		}
 
 	private:
-		FString	ImageName;
 		FAutomationScreenshotData ComparisonParameters;
 		const TArray<FColor> ImageData;
 		bool Initiate;
