Fix missing #undef LOCTEXT_NAMESPACE

#rnx

diff --git a/Engine/Source/Developer/ScreenShotComparisonTools/Private/ScreenShotManager.cpp b/Engine/Source/Developer/ScreenShotComparisonTools/Private/ScreenShotManager.cpp
--- a/Engine/Source/Developer/ScreenShotComparisonTools/Private/ScreenShotManager.cpp
+++ b/Engine/Source/Developer/ScreenShotComparisonTools/Private/ScreenShotManager.cpp
@@ -806,3 +806,5 @@ void FScreenShotManager::BuildFallbackPlatformsListFromConfig()
 		}
 	}
 }
+
+#undef LOCTEXT_NAMESPACE
