Remove screenshot legacy path check + fix EApprovedFolderOptions values
#rnx

diff --git a/Engine/Source/Developer/ScreenShotComparisonTools/Private/ScreenShotManager.cpp b/Engine/Source/Developer/ScreenShotComparisonTools/Private/ScreenShotManager.cpp
--- a/Engine/Source/Developer/ScreenShotComparisonTools/Private/ScreenShotManager.cpp
+++ b/Engine/Source/Developer/ScreenShotComparisonTools/Private/ScreenShotManager.cpp
@@ -96,7 +96,7 @@ FString FScreenShotManager::GetApprovedFolderForImageWithOptions(const FAutomati
 {
 	const FDataDrivenPlatformInfo& PlatInfo = FDataDrivenPlatformInfoRegistry::GetPlatformInfo(MetaData.Platform);
 
-	bool bUsePlatformPath = PlatInfo.bIsConfidential && (InOptions & EApprovedFolderOptions::UsePlatformFolders) == 0;
+	bool bUsePlatformPath = PlatInfo.bIsConfidential && (InOptions & EApprovedFolderOptions::UsePlatformFolders) != 0;
 
 	// Test folder will be MapOrContext/ImageName
 	FString TestFolder = GetPathComponentForTestImages(MetaData, false);
@@ -132,7 +132,7 @@ FString FScreenShotManager::GetApprovedFolderForImageWithOptions(const FAutomati
 
 FString FScreenShotManager::GetIdealApprovedFolderForImage(const FAutomationScreenshotMetadata& MetaData) const
 {
-	EApprovedFolderOptions DefaultOptions = bUseConfidentialPlatformPaths ? EApprovedFolderOptions::None : EApprovedFolderOptions::UsePlatformFolders;
+	EApprovedFolderOptions DefaultOptions = bUseConfidentialPlatformPaths ? EApprovedFolderOptions::UsePlatformFolders : EApprovedFolderOptions::None;
 	return GetApprovedFolderForImageWithOptions(MetaData, DefaultOptions);
 }
 
@@ -148,7 +148,7 @@ TArray<FString> FScreenShotManager::FindApprovedFiles(const FAutomationScreensho
 
 	TArray<FString> ApprovedImages;
 
-	EApprovedFolderOptions Options = bUseConfidentialPlatformPaths ? EApprovedFolderOptions::None : EApprovedFolderOptions::UsePlatformFolders;
+	EApprovedFolderOptions Options = bUseConfidentialPlatformPaths ? EApprovedFolderOptions::UsePlatformFolders : EApprovedFolderOptions::None;
 
 	// check out standard path using whether confidential platforms are in a separate tree
 	FString ApprovedPath = GetApprovedFolderForImageWithOptions(IncomingMetaData, Options);
@@ -157,26 +157,12 @@ TArray<FString> FScreenShotManager::FindApprovedFiles(const FAutomationScreensho
 	// Make sure the first log line is of the first path tried, not the last fallback. The list of fallbacks will be printed if nothing is found.
 	const FString FirstApprovedPath = ApprovedPath;
 
-	// check again, but try legacy paths
-	if (!ApprovedImages.Num())
-	{
-		ApprovedPath = GetApprovedFolderForImageWithOptions(IncomingMetaData, Options | EApprovedFolderOptions::UseLegacyPaths);
-		FindImages(ApprovedImages, ApprovedPath);
-	}
-
 	// if we're a blank and bUseConfidentialPlatformPaths, try without that
 	if (ApprovedImages.Num() == 0 && bUseConfidentialPlatformPaths)
 	{
-		// check legacy paths.
+		// check standard paths.
 		ApprovedPath = FPaths::GetPath(GetApprovedFolderForImageWithOptions(IncomingMetaData, EApprovedFolderOptions::None));
 		FindImages(ApprovedImages, ApprovedPath);
-
-		// check again, but try legacy paths
-		if (!ApprovedImages.Num())
-		{
-			ApprovedPath = GetApprovedFolderForImageWithOptions(IncomingMetaData, EApprovedFolderOptions::UseLegacyPaths);
-			FindImages(ApprovedImages, ApprovedPath);
-		}
 	}
 
 	// find fallback images if they don't exist at this point
@@ -221,13 +207,6 @@ TArray<FString> FScreenShotManager::FindApprovedFiles(const FAutomationScreensho
 				ApprovedPath = FPaths::GetPath(GetIdealApprovedFolderForImage(CopiedMetaData));
 				FindImages(ApprovedImages, ApprovedPath);
 
-				// check again, but try legacy paths
-				if (!ApprovedImages.Num())
-				{
-					ApprovedPath = GetApprovedFolderForImageWithOptions(CopiedMetaData, EApprovedFolderOptions::UseLegacyPaths);
-					FindImages(ApprovedImages, ApprovedPath);
-				}
-
 				if (ApprovedImages.Num())
 				{
 					UE_LOG(LogScreenShotManager, Log, TEXT("Using fallback images from %s"), *ApprovedPath);
@@ -403,6 +382,11 @@ FImageComparisonResult FScreenShotManager::CompareScreenshot(const FString& InUn
 
 	// Result paths should be relative to the project. Note this may be empty, and if it is MakePathRelative returns
 	// a non empty relative path... but we want it to stay empty as that's how we signal that no approved file exists
+
+	// these two must exist...
+	FPaths::MakePathRelativeTo(ComparisonResult.IncomingFilePath, *FPaths::ProjectDir());
+	FPaths::MakePathRelativeTo(ComparisonResult.IdealApprovedFolderPath, *FPaths::ProjectDir());
+
 	if (!ComparisonResult.ApprovedFilePath.IsEmpty())
 	{
 		FPaths::MakePathRelativeTo(ComparisonResult.ApprovedFilePath, *FPaths::ProjectDir());
@@ -413,10 +397,6 @@ FImageComparisonResult FScreenShotManager::CompareScreenshot(const FString& InUn
 		FPaths::MakePathRelativeTo(ComparisonResult.ComparisonFilePath, *FPaths::ProjectDir());
 	}
 
-	// these two must exist...
-	FPaths::MakePathRelativeTo(ComparisonResult.IncomingFilePath, *FPaths::ProjectDir());
-	FPaths::MakePathRelativeTo(ComparisonResult.IdealApprovedFolderPath, *FPaths::ProjectDir());
-
 	// Report path is something like Test/Context/ that can be freely moved around and relocated with the relative paths
 	// to the data remaining intact
 	FString ReportPathOnDisk = FPaths::Combine(ScreenshotResultsFolder, ResultsSubFolder, IncomingMetaData.Platform, GetPathComponentForRHI(IncomingMetaData), TEXT("/"));
diff --git a/Engine/Source/Developer/ScreenShotComparisonTools/Private/ScreenShotManager.h b/Engine/Source/Developer/ScreenShotComparisonTools/Private/ScreenShotManager.h
--- a/Engine/Source/Developer/ScreenShotComparisonTools/Private/ScreenShotManager.h
+++ b/Engine/Source/Developer/ScreenShotComparisonTools/Private/ScreenShotManager.h
@@ -28,8 +28,7 @@ public:
 	enum class EApprovedFolderOptions : int8
 	{
 		None = 0,
-		UsePlatformFolders	= 1 >> 0,
-		UseLegacyPaths		= 1 >> 1
+		UsePlatformFolders	= 1 << 0
 	};
 
 	/**
