Automation: Make image sequence comparisons respect the comparison rules inside of the Metadata structure, if bHasComparisonRules is set to true. Movie Render Queue automation tests now expose the FImageTolerance presets to the UI through an enum, and through custom values.

diff --git a/Engine/Plugins/MovieScene/MovieRenderPipeline/Source/MovieRenderPipelineEditor/Private/MoviePipelineFunctionalTestBase.cpp b/Engine/Plugins/MovieScene/MovieRenderPipeline/Source/MovieRenderPipelineEditor/Private/MoviePipelineFunctionalTestBase.cpp
--- a/Engine/Plugins/MovieScene/MovieRenderPipeline/Source/MovieRenderPipelineEditor/Private/MoviePipelineFunctionalTestBase.cpp
+++ b/Engine/Plugins/MovieScene/MovieRenderPipeline/Source/MovieRenderPipelineEditor/Private/MoviePipelineFunctionalTestBase.cpp
@@ -24,9 +24,10 @@
 
 AMoviePipelineFunctionalTestBase::AMoviePipelineFunctionalTestBase()
 {
-	
+	ImageToleranceLevel = EImageTolerancePreset::IgnoreLess;
 }
 
+
 void AMoviePipelineFunctionalTestBase::PrepareTest() 
 {
 	if (!QueuePreset)
@@ -156,7 +157,20 @@ void AMoviePipelineFunctionalTestBase::CompareRenderOutputToGroundTruth(FMoviePi
 	FAutomationScreenshotData Data = UAutomationBlueprintFunctionLibrary::BuildScreenshotData(GetWorld(), TestLabel, OutputSetting->OutputResolution.X, OutputSetting->OutputResolution.Y);
 
 	// Convert the Screenshot Data into Metadata
+	const FImageTolerance ImageTolerance = GetImageToleranceForPreset(ImageToleranceLevel, CustomImageTolerance);
+
 	FAutomationScreenshotMetadata MetaData(Data);
+	MetaData.bHasComparisonRules = true;
+	MetaData.ToleranceRed = ImageTolerance.Red;
+	MetaData.ToleranceGreen = ImageTolerance.Green;
+	MetaData.ToleranceBlue = ImageTolerance.Blue;
+	MetaData.ToleranceAlpha = ImageTolerance.Alpha;
+	MetaData.ToleranceMinBrightness = ImageTolerance.MinBrightness;
+	MetaData.ToleranceMaxBrightness = ImageTolerance.MaxBrightness;
+	MetaData.bIgnoreAntiAliasing = ImageTolerance.IgnoreAntiAliasing;
+	MetaData.bIgnoreColors = ImageTolerance.IgnoreColors;
+	MetaData.MaximumLocalError = ImageTolerance.MaximumLocalError;
+	MetaData.MaximumGlobalError = ImageTolerance.MaximumGlobalError;
 
 	IScreenShotToolsModule& ScreenShotModule = FModuleManager::LoadModuleChecked<IScreenShotToolsModule>("ScreenShotComparisonTools");
 	IScreenShotManagerPtr ScreenshotManager = ScreenShotModule.GetScreenShotManager();
diff --git a/Engine/Plugins/MovieScene/MovieRenderPipeline/Source/MovieRenderPipelineEditor/Private/MoviePipelineFunctionalTestBase.h b/Engine/Plugins/MovieScene/MovieRenderPipeline/Source/MovieRenderPipelineEditor/Private/MoviePipelineFunctionalTestBase.h
--- a/Engine/Plugins/MovieScene/MovieRenderPipeline/Source/MovieRenderPipelineEditor/Private/MoviePipelineFunctionalTestBase.h
+++ b/Engine/Plugins/MovieScene/MovieRenderPipeline/Source/MovieRenderPipelineEditor/Private/MoviePipelineFunctionalTestBase.h
@@ -7,6 +7,7 @@
 #include "FunctionalTest.h"
 #include "Misc/AutomationTest.h"
 #include "MovieRenderPipelineDataTypes.h"
+#include "ImageComparer.h"
 
 #include "MoviePipelineFunctionalTestBase.generated.h"
 
@@ -45,6 +46,12 @@ protected:
 	UPROPERTY(EditAnywhere, BlueprintReadOnly, Category = "Movie Pipeline")
 	TObjectPtr<UMoviePipelineQueue> QueuePreset;
 
+	UPROPERTY(EditAnywhere, Category = "Movie Pipeline")
+	EImageTolerancePreset ImageToleranceLevel;
+
+	UPROPERTY(EditAnywhere, Category = "Movie Pipeline")
+	FImageTolerance CustomImageTolerance;
+
 protected:
 	UPROPERTY(Transient)
 	TObjectPtr<UMoviePipeline> ActiveMoviePipeline;
diff --git a/Engine/Source/Developer/ScreenShotComparisonTools/Private/ScreenShotManager.cpp b/Engine/Source/Developer/ScreenShotComparisonTools/Private/ScreenShotManager.cpp
--- a/Engine/Source/Developer/ScreenShotComparisonTools/Private/ScreenShotManager.cpp
+++ b/Engine/Source/Developer/ScreenShotComparisonTools/Private/ScreenShotManager.cpp
@@ -658,9 +658,25 @@ TSharedPtr<FImageComparisonResult> FScreenShotManager::CompareImageSequence(cons
 		FString OldExtension = FPaths::GetExtension(NewImagePath, true);
 		DeltaPath += TEXT("_Delta") + OldExtension;
 
-		// Alright we have both images in memory now, now use a FImageComparer for less strict comparison.
-		FImageTolerance Tolerance = FImageTolerance::DefaultIgnoreLess;
+		// If the metadata for the screenshot does not provide tolerance rules, use these instead.
+		const FImageTolerance DefaultTolerance = FImageTolerance::DefaultIgnoreLess;
 		FImageComparer Comparer;
+		FImageTolerance Tolerance = DefaultTolerance;
+
+		if (Metadata.bHasComparisonRules)
+		{
+			Tolerance.Red = Metadata.ToleranceRed;
+			Tolerance.Green = Metadata.ToleranceGreen;
+			Tolerance.Blue = Metadata.ToleranceBlue;
+			Tolerance.Alpha = Metadata.ToleranceAlpha;
+			Tolerance.MinBrightness = Metadata.ToleranceMinBrightness;
+			Tolerance.MaxBrightness = Metadata.ToleranceMaxBrightness;
+			Tolerance.IgnoreAntiAliasing = Metadata.bIgnoreAntiAliasing;
+			Tolerance.IgnoreColors = Metadata.bIgnoreColors;
+			Tolerance.MaximumLocalError = Metadata.MaximumLocalError;
+			Tolerance.MaximumGlobalError = Metadata.MaximumGlobalError;
+		}
+
 		FImageComparisonResult ComparisonResult = Comparer.Compare(OldImagePath, NewImagePath, Tolerance, DeltaPath);
 
 		if (!ComparisonResult.AreSimilar())
diff --git a/Engine/Source/Developer/ScreenShotComparisonTools/Public/ImageComparer.h b/Engine/Source/Developer/ScreenShotComparisonTools/Public/ImageComparer.h
--- a/Engine/Source/Developer/ScreenShotComparisonTools/Public/ImageComparer.h
+++ b/Engine/Source/Developer/ScreenShotComparisonTools/Public/ImageComparer.h
@@ -29,34 +29,34 @@ struct SCREENSHOTCOMPARISONTOOLS_API FImageTolerance
 
 public:
 
-	UPROPERTY()
+	UPROPERTY(EditAnywhere, Category = Automation)
 	uint8 Red;
 
-	UPROPERTY()
+	UPROPERTY(EditAnywhere, Category = Automation)
 	uint8 Green;
 
-	UPROPERTY()
+	UPROPERTY(EditAnywhere, Category = Automation)
 	uint8 Blue;
 
-	UPROPERTY()
+	UPROPERTY(EditAnywhere, Category = Automation)
 	uint8 Alpha;
 
-	UPROPERTY()
+	UPROPERTY(EditAnywhere, Category = Automation)
 	uint8 MinBrightness;
 
-	UPROPERTY()
+	UPROPERTY(EditAnywhere, Category = Automation)
 	uint8 MaxBrightness;
 
-	UPROPERTY()
+	UPROPERTY(EditAnywhere, Category = Automation)
 	bool IgnoreAntiAliasing;
 
-	UPROPERTY()
+	UPROPERTY(EditAnywhere, Category = Automation)
 	bool IgnoreColors;
 
-	UPROPERTY()
+	UPROPERTY(EditAnywhere, Category = Automation)
 	float MaximumLocalError;
 
-	UPROPERTY()
+	UPROPERTY(EditAnywhere, Category = Automation)
 	float MaximumGlobalError;
 
 	FImageTolerance()
@@ -94,6 +94,32 @@ public:
 	const static FImageTolerance DefaultIgnoreColors;
 };
 
+UENUM()
+enum class EImageTolerancePreset : uint8
+{
+	IgnoreNothing,
+	IgnoreLess,
+	IgnoreAntiAliasing,
+	IgnoreColors,
+	Custom
+};
+
+static FImageTolerance GetImageToleranceForPreset(EImageTolerancePreset Preset, FImageTolerance CustomTolerance)
+{
+	switch (Preset)
+	{
+	case EImageTolerancePreset::IgnoreNothing: return FImageTolerance::DefaultIgnoreNothing;
+	case EImageTolerancePreset::IgnoreLess: return FImageTolerance::DefaultIgnoreLess;
+	case EImageTolerancePreset::IgnoreAntiAliasing: return FImageTolerance::DefaultIgnoreAntiAliasing;
+	case EImageTolerancePreset::IgnoreColors: return FImageTolerance::DefaultIgnoreColors;
+	case EImageTolerancePreset::Custom: return CustomTolerance;
+	default:
+		check(0);
+	}
+
+	return FImageTolerance();
+}
+
 class FComparableImage;
 
 class FPixelOperations
