Fix incorrect use of "Alternative Results" screenshot tests.

This CL adds a flag to track if a screenshot reference was taken or not with Substrate enabled. Alternative results are not compared one by one until one succeed as I thought, but use the 'closest' image based in its metadata. This can the Legacy path fails for test cases which were updated with an alternative results specific to Substrate. This CL adds a metadata to screenshot reference in order to ensure that reference image taken with Subtrate anre not used in legacy case.

diff --git a/Engine/Source/Developer/ScreenShotComparisonTools/Private/ScreenShotManager.cpp b/Engine/Source/Developer/ScreenShotComparisonTools/Private/ScreenShotManager.cpp
--- a/Engine/Source/Developer/ScreenShotComparisonTools/Private/ScreenShotManager.cpp
+++ b/Engine/Source/Developer/ScreenShotComparisonTools/Private/ScreenShotManager.cpp
@@ -18,12 +18,19 @@
 #include "Misc/DataDrivenPlatformInfoRegistry.h"
 #include "ScreenShotComparisonSettings.h"
 #include "ILauncherServicesModule.h"
-
+#include "RenderUtils.h"
 
 DEFINE_LOG_CATEGORY(LogScreenShotManager);
 
 #define LOCTEXT_NAMESPACE "ScreenShotManager"
 
+/* Return true of the Reference screenshot is compatible with the incoming one */
+static bool IsCompatible(const FAutomationScreenshotMetadata& Reference, const FAutomationScreenshotMetadata& Incoming)
+{
+	// If Substrate is disabled and the screenshot was approved with Substrate enabled, mark it as incompatible as it might contains result specific to Substrate.
+	return Substrate::IsSubstrateEnabled() ? true : !Reference.bIsSubstrate;
+}
+
 class FScreenshotComparisons
 {
 public:
@@ -242,7 +249,6 @@ TFuture<FImageComparisonResult> FScreenShotManager::CompareScreenshotAsync(const
 	return Async(EAsyncExecution::Thread, [this, IncomingPath, MetaData, Options] () { return CompareScreenshot(IncomingPath, MetaData, Options); });
 }
 
-
 /**
  * Takes an incoming file and returns a comparison based on comparing it to a reference image that exists in the project.
  * The comparison results (incoming image, approved image, comparison image, report metadata) will be placed in the provided
@@ -317,6 +323,13 @@ FImageComparisonResult FScreenShotManager::CompareScreenshot(const FString& InUn
 					FAutomationScreenshotMetadata Metadata;
 					if ( FJsonObjectConverter::JsonObjectStringToUStruct(Json, &Metadata, 0, 0) )
 					{
+						// If there is only one reference screenshot, use it.
+						const bool bCompatibleScreenShot = ApprovedDeviceShots.Num() == 1 || IsCompatible(Metadata, ExistingMetadata.GetValue());
+						if (!bCompatibleScreenShot)
+						{
+							continue;
+						}
+
 						int32 Comparison = Metadata.Compare(ExistingMetadata.GetValue());
 						if ( Comparison > MatchScore )
 						{
diff --git a/Engine/Source/Developer/ScreenShotComparisonTools/ScreenShotComparisonTools.Build.cs b/Engine/Source/Developer/ScreenShotComparisonTools/ScreenShotComparisonTools.Build.cs
--- a/Engine/Source/Developer/ScreenShotComparisonTools/ScreenShotComparisonTools.Build.cs
+++ b/Engine/Source/Developer/ScreenShotComparisonTools/ScreenShotComparisonTools.Build.cs
@@ -20,7 +20,8 @@ public class ScreenShotComparisonTools : ModuleRules
 				"ImageWrapper",
 				"Json",
 				"JsonUtilities",
-				"DesktopPlatform"
+				"DesktopPlatform",
+				"RenderCore"
 			}
 		);
 
diff --git a/Engine/Source/Runtime/AutomationMessages/Public/AutomationWorkerMessages.h b/Engine/Source/Runtime/AutomationMessages/Public/AutomationWorkerMessages.h
--- a/Engine/Source/Runtime/AutomationMessages/Public/AutomationWorkerMessages.h
+++ b/Engine/Source/Runtime/AutomationMessages/Public/AutomationWorkerMessages.h
@@ -454,6 +454,10 @@ public:
 	UPROPERTY(EditAnywhere, Category="Message")
 	bool bIsStereo;
 
+	// Track if ScreenShotMetadata was created with Substrate enabled or not
+	UPROPERTY(EditAnywhere, Category="Message")
+	bool bIsSubstrate;
+
 	// Hardware Details
 	UPROPERTY(EditAnywhere, Category="Message")
 	FString Vendor;
@@ -520,6 +524,7 @@ public:
 		: Width(0)
 		, Height(0)
 		, bIsStereo(false)
+		, bIsSubstrate(false)
 		, ResolutionQuality(0.0f)
 		, ViewDistanceQuality(0)
 		, AntiAliasingQuality(0)
@@ -569,6 +574,7 @@ public:
 		Platform = Data.Platform;
 		FeatureLevel = Data.FeatureLevel;
 		bIsStereo = Data.bIsStereo;
+		bIsSubstrate = Data.bIsSubstrate;
 
 		// Hardware Details
 		Vendor = Data.Vendor;
diff --git a/Engine/Source/Runtime/Core/Public/Misc/AutomationTest.h b/Engine/Source/Runtime/Core/Public/Misc/AutomationTest.h
--- a/Engine/Source/Runtime/Core/Public/Misc/AutomationTest.h
+++ b/Engine/Source/Runtime/Core/Public/Misc/AutomationTest.h
@@ -737,6 +737,9 @@ struct FAutomationScreenshotData
 	FString FeatureLevel;
 	bool bIsStereo;
 
+	// Track if ScreenShotData was created with Substrate enabled or not
+	bool bIsSubstrate;
+
 	// Hardware Details
 	FString Vendor;
 	FString AdapterName;
@@ -779,6 +782,7 @@ struct FAutomationScreenshotData
 		, Width(0)
 		, Height(0)
 		, bIsStereo(false)
+		, bIsSubstrate(false)
 		, ResolutionQuality(1.0f)
 		, ViewDistanceQuality(0)
 		, AntiAliasingQuality(0)
diff --git a/Engine/Source/Runtime/Engine/Private/Tests/AutomationCommon.cpp b/Engine/Source/Runtime/Engine/Private/Tests/AutomationCommon.cpp
--- a/Engine/Source/Runtime/Engine/Private/Tests/AutomationCommon.cpp
+++ b/Engine/Source/Runtime/Engine/Private/Tests/AutomationCommon.cpp
@@ -25,6 +25,7 @@
 #include "IRenderCaptureProvider.h"
 #include "Modules/ModuleManager.h"
 #include "Algo/Accumulate.h"
+#include "RenderUtils.h"
 
 #if WITH_AUTOMATION_TESTS
 
@@ -488,6 +489,7 @@ namespace AutomationCommon
 		Data.Rhi = FHardwareInfo::GetHardwareInfo(NAME_RHI);
 		GetFeatureLevelName(GMaxRHIFeatureLevel, Data.FeatureLevel);
 		Data.bIsStereo = GEngine->StereoRenderingDevice.IsValid() ? GEngine->StereoRenderingDevice->IsStereoEnabled() : false;
+		Data.bIsSubstrate = Substrate::IsSubstrateEnabled();
 		Data.Vendor = RHIVendorIdToString();
 		Data.AdapterName = GRHIAdapterName;
 		Data.AdapterInternalDriverVersion = GRHIAdapterInternalDriverVersion;
